Fix cooking errors

diff --git a/Engine/Plugins/Runtime/ApexDestruction/Source/ApexDestruction/Private/ApexDestructionModule.cpp b/Engine/Plugins/Runtime/ApexDestruction/Source/ApexDestruction/Private/ApexDestructionModule.cpp
--- a/Engine/Plugins/Runtime/ApexDestruction/Source/ApexDestruction/Private/ApexDestructionModule.cpp
+++ b/Engine/Plugins/Runtime/ApexDestruction/Source/ApexDestruction/Private/ApexDestructionModule.cpp
@@ -192,10 +192,12 @@ public:
 	virtual void OnInitPhys(FPhysScene* PhysScene, EPhysicsSceneType SceneType)
 	{
 #if WITH_APEX
-		if(PxScene* PScene = PhysScene->GetPhysXScene(SceneType))
+		//Can't set this flag after scene already created
+		/*if(PxScene* PScene = PhysScene->GetPhysXScene(SceneType))
 		{
+			SCOPED_SCENE_WRITE_LOCK(PScene);
 			PScene->setFlag(PxSceneFlag::eENABLE_KINEMATIC_PAIRS, true);
-		}
+		}*/
 
 		if((PhysScene->bAsyncSceneEnabled && SceneType == PST_Async) || (!PhysScene->bAsyncSceneEnabled && SceneType == PST_Sync))
 		{
diff --git a/Engine/Source/Runtime/Engine/Private/PhysicsEngine/PhysScene.cpp b/Engine/Source/Runtime/Engine/Private/PhysicsEngine/PhysScene.cpp
--- a/Engine/Source/Runtime/Engine/Private/PhysicsEngine/PhysScene.cpp
+++ b/Engine/Source/Runtime/Engine/Private/PhysicsEngine/PhysScene.cpp
@@ -1951,6 +1951,7 @@ void FPhysScene::InitPhysScene(uint32 SceneType)
 
 	// Need to turn this on to consider kinematics turning into dynamic. Otherwise, you'll need to call resetFiltering to do the expensive broadphase reinserting 
 	PSceneDesc.flags |= PxSceneFlag::eENABLE_KINEMATIC_STATIC_PAIRS;
+	PSceneDesc.flags |= PxSceneFlag::eENABLE_KINEMATIC_PAIRS;	//this is only needed for destruction, but unfortunately this flag cannot be modified after creation and the plugin has no hook (yet)
 
 	// @TODO Should we set up PSceneDesc.limits? How?
 
