Other fixes for PVS studio that were masked by previous warnings (tested by preflight)

diff --git a/Engine/Plugins/Online/OnlineSubsystemOculus/Source/Private/OculusNetDriver.cpp b/Engine/Plugins/Online/OnlineSubsystemOculus/Source/Private/OculusNetDriver.cpp
--- a/Engine/Plugins/Online/OnlineSubsystemOculus/Source/Private/OculusNetDriver.cpp
+++ b/Engine/Plugins/Online/OnlineSubsystemOculus/Source/Private/OculusNetDriver.cpp
@@ -292,8 +292,8 @@ void UOculusNetDriver::LowLevelSend(FString Address, void* Data, int32 CountBits
 	{
 		return UIpNetDriver::LowLevelSend(Address, Data, CountBits);
 	}
-	auto OculusAddr = new FInternetAddrOculus(FURL(nullptr, *Address, ETravelType::TRAVEL_Absolute));
-	ovrID PeerID = OculusAddr->GetID();
+	FInternetAddrOculus OculusAddr(FURL(nullptr, *Address, ETravelType::TRAVEL_Absolute));
+	ovrID PeerID = OculusAddr.GetID();
 	if (ovr_Net_IsConnected(PeerID))
 	{
 		const uint8* DataToSend = reinterpret_cast<uint8*>(Data);
diff --git a/Engine/Plugins/Runtime/ApexDestruction/Source/ApexDestruction/Private/ApexDestructibleAssetImport.cpp b/Engine/Plugins/Runtime/ApexDestruction/Source/ApexDestruction/Private/ApexDestructibleAssetImport.cpp
--- a/Engine/Plugins/Runtime/ApexDestruction/Source/ApexDestruction/Private/ApexDestructibleAssetImport.cpp
+++ b/Engine/Plugins/Runtime/ApexDestruction/Source/ApexDestruction/Private/ApexDestructibleAssetImport.cpp
@@ -830,6 +830,12 @@ bool SetApexDestructibleAsset(UDestructibleMesh& DestructibleMesh, apex::Destruc
 	bool bHaveNormals, bHaveTangents;
 	if (!FillSkelMeshImporterFromApexDestructibleAsset(*SkelMeshImportDataPtr, ApexDestructibleAsset, bHaveNormals, bHaveTangents))
 	{
+		if (ExistDestMeshDataPtr)
+		{
+			RestoreExistingDestMeshData(ExistDestMeshDataPtr, &DestructibleMesh);
+			delete ExistDestMeshDataPtr;
+			ExistDestMeshDataPtr = NULL;
+		}
 		return false;
 	}
 
@@ -910,6 +916,10 @@ bool SetApexDestructibleAsset(UDestructibleMesh& DestructibleMesh, apex::Destruc
 		if (!MeshUtilities.BuildSkeletalMesh(DestructibleMeshResource.LODModels[0], DestructibleMesh.RefSkeleton, LODInfluences,LODWedges,LODFaces,LODPoints,LODPointToRawMap,BuildOptions))
 		{
 			DestructibleMesh.MarkPendingKill();
+
+			delete ExistDestMeshDataPtr;
+			ExistDestMeshDataPtr = NULL;
+
 			return false;
 		}
 
