Integrate DestructibleMesh Component Fix to 4.20

diff --git a/Engine/Plugins/Runtime/ApexDestruction/Source/ApexDestruction/Private/DestructibleComponent.cpp b/Engine/Plugins/Runtime/ApexDestruction/Source/ApexDestruction/Private/DestructibleComponent.cpp
--- a/Engine/Plugins/Runtime/ApexDestruction/Source/ApexDestruction/Private/DestructibleComponent.cpp
+++ b/Engine/Plugins/Runtime/ApexDestruction/Source/ApexDestruction/Private/DestructibleComponent.cpp
@@ -950,6 +950,8 @@ void UDestructibleComponent::SetChunksWorldTM(const TArray<FUpdateChunksInfo>& U
 		GetEditableComponentSpaceTransforms()[BoneIndex] = FTransform(BoneRotation, BoneTranslation);
 	}
 
+	bNeedToFlipSpaceBaseBuffers = true;
+
 	// Mark the transform as dirty, so the bounds are updated and sent to the render thread
 	MarkRenderTransformDirty();
 
