Disable apex destructible for HoloLens, as Apex Destruction is no longer supported by Nvidia

diff --git a/Engine/Plugins/Runtime/ApexDestruction/ApexDestruction.uplugin b/Engine/Plugins/Runtime/ApexDestruction/ApexDestruction.uplugin
--- a/Engine/Plugins/Runtime/ApexDestruction/ApexDestruction.uplugin
+++ b/Engine/Plugins/Runtime/ApexDestruction/ApexDestruction.uplugin
@@ -16,11 +16,14 @@
 	"Installed" : false,
 	"Modules" :
 	[
-    {
-      "Name": "ApexDestruction",
-      "Type": "Runtime",
-      "LoadingPhase": "PreDefault"
-    },
+		{
+			"Name": "ApexDestruction",
+			"Type": "Runtime",
+			"LoadingPhase": "PreDefault",
+			"BlacklistPlatforms": [
+				"HoloLens"
+			]
+		},
     {
 			"Name" : "ApexDestructionEditor",
 			"Type" : "Developer",
