Avoid dereferencing null.

diff --git a/Engine/Plugins/Runtime/ApexDestruction/Source/ApexDestruction/Private/DestructibleMesh.cpp b/Engine/Plugins/Runtime/ApexDestruction/Source/ApexDestruction/Private/DestructibleMesh.cpp
--- a/Engine/Plugins/Runtime/ApexDestruction/Source/ApexDestruction/Private/DestructibleMesh.cpp
+++ b/Engine/Plugins/Runtime/ApexDestruction/Source/ApexDestruction/Private/DestructibleMesh.cpp
@@ -677,7 +677,8 @@ bool UDestructibleMesh::BuildFractureSettingsFromStaticMesh(UStaticMesh* StaticM
 			int32 SectionMaterialIndex = INDEX_NONE;
 			for (int32 MaterialIndex = 0; MaterialIndex < MeshMaterials.Num(); ++MaterialIndex)
 			{
-				if (MeshMaterials[MaterialIndex]->GetFName() == CurrentStaticMesh->GetMaterial(Section.MaterialIndex)->GetFName())
+				UMaterialInterface* CurrentMeshMaterial = CurrentStaticMesh->GetMaterial(Section.MaterialIndex);
+				if (MeshMaterials[MaterialIndex] && CurrentMeshMaterial && (MeshMaterials[MaterialIndex]->GetFName() == CurrentMeshMaterial->GetFName()))
 				{
 					SectionMaterialIndex = MaterialIndex;
 					break;
