Avoid crash with null materials

diff --git a/Engine/Plugins/Runtime/ApexDestruction/Source/ApexDestruction/Private/ApexDestructibleAssetImport.cpp b/Engine/Plugins/Runtime/ApexDestruction/Source/ApexDestruction/Private/ApexDestructibleAssetImport.cpp
--- a/Engine/Plugins/Runtime/ApexDestruction/Source/ApexDestruction/Private/ApexDestructibleAssetImport.cpp
+++ b/Engine/Plugins/Runtime/ApexDestruction/Source/ApexDestruction/Private/ApexDestructibleAssetImport.cpp
@@ -994,23 +994,36 @@ bool BuildDestructibleMeshFromFractureSettings(UDestructibleMesh& DestructibleMe
 			OverrideMaterials[MaterialIndex] = DestructibleMesh.Materials[MaterialIndex].MaterialInterface;
 		}
 
-		DestructibleMesh.Materials.SetNum(DestructibleMesh.FractureSettings->Materials.Num());
+		int32 MaterialCount = 0;
+		for (int32 MaterialIndex = 0; MaterialIndex < DestructibleMesh.Materials.Num(); ++MaterialIndex)
+		{
+			if (DestructibleMesh.FractureSettings->Materials[MaterialIndex] != nullptr)
+			{
+				MaterialCount++;
+			}
+		}
+		DestructibleMesh.Materials.SetNum(MaterialCount);
 
+		int32 NewMaterialIndex = 0;
 		for (int32 MaterialIndex = 0; MaterialIndex < DestructibleMesh.Materials.Num(); ++MaterialIndex)
 		{
+			if (DestructibleMesh.FractureSettings->Materials[MaterialIndex] == nullptr)
+			{
+				continue;
+			}
 			if(MaterialIndex < OverrideMaterials.Num() && OverrideMaterials[MaterialIndex])//if user has overridden materials use it
 			{
-				DestructibleMesh.Materials[MaterialIndex].MaterialInterface = OverrideMaterials[MaterialIndex];
-				DestructibleMesh.Materials[MaterialIndex].ImportedMaterialSlotName = OverrideMaterials[MaterialIndex]->GetFName();
-				DestructibleMesh.Materials[MaterialIndex].MaterialSlotName = OverrideMaterials[MaterialIndex]->GetFName();
+				DestructibleMesh.Materials[NewMaterialIndex].MaterialInterface = OverrideMaterials[MaterialIndex];
+				DestructibleMesh.Materials[NewMaterialIndex].ImportedMaterialSlotName = OverrideMaterials[MaterialIndex]->GetFName();
+				DestructibleMesh.Materials[NewMaterialIndex].MaterialSlotName = OverrideMaterials[MaterialIndex]->GetFName();
 			}
 			else
 			{
-				DestructibleMesh.Materials[MaterialIndex].MaterialInterface = DestructibleMesh.FractureSettings->Materials[MaterialIndex];
-				DestructibleMesh.Materials[MaterialIndex].ImportedMaterialSlotName = DestructibleMesh.FractureSettings->Materials[MaterialIndex]->GetFName();
-				DestructibleMesh.Materials[MaterialIndex].MaterialSlotName = DestructibleMesh.FractureSettings->Materials[MaterialIndex]->GetFName();
+				DestructibleMesh.Materials[NewMaterialIndex].MaterialInterface = DestructibleMesh.FractureSettings->Materials[MaterialIndex];
+				DestructibleMesh.Materials[NewMaterialIndex].ImportedMaterialSlotName = DestructibleMesh.FractureSettings->Materials[MaterialIndex]->GetFName();
+				DestructibleMesh.Materials[NewMaterialIndex].MaterialSlotName = DestructibleMesh.FractureSettings->Materials[MaterialIndex]->GetFName();
 			}
-
+			NewMaterialIndex++;
 		}
 
 		apex::DestructibleAssetCookingDesc DestructibleAssetCookingDesc;
