Fix the material assignation for destructible mesh
Force import material to the same value has the mesh use to create the destructible mesh.

diff --git a/Engine/Plugins/Runtime/ApexDestruction/Source/ApexDestruction/Private/ApexDestructibleAssetImport.cpp b/Engine/Plugins/Runtime/ApexDestruction/Source/ApexDestruction/Private/ApexDestructibleAssetImport.cpp
--- a/Engine/Plugins/Runtime/ApexDestruction/Source/ApexDestruction/Private/ApexDestructibleAssetImport.cpp
+++ b/Engine/Plugins/Runtime/ApexDestruction/Source/ApexDestruction/Private/ApexDestructibleAssetImport.cpp
@@ -995,36 +995,27 @@ bool BuildDestructibleMeshFromFractureSettings(UDestructibleMesh& DestructibleMe
 			OverrideMaterials[MaterialIndex] = DestructibleMesh.Materials[MaterialIndex].MaterialInterface;
 		}
 
-		int32 MaterialCount = 0;
-		for (int32 MaterialIndex = 0; MaterialIndex < DestructibleMesh.Materials.Num(); ++MaterialIndex)
-		{
-			if (DestructibleMesh.FractureSettings->Materials[MaterialIndex] != nullptr)
-			{
-				MaterialCount++;
-			}
-		}
+		int32 MaterialCount = DestructibleMesh.FractureSettings->Materials.Num();
+
 		DestructibleMesh.Materials.SetNum(MaterialCount);
 
-		int32 NewMaterialIndex = 0;
 		for (int32 MaterialIndex = 0; MaterialIndex < DestructibleMesh.Materials.Num(); ++MaterialIndex)
 		{
-			if (DestructibleMesh.FractureSettings->Materials[MaterialIndex] == nullptr)
-			{
-				continue;
-			}
-			if(MaterialIndex < OverrideMaterials.Num() && OverrideMaterials[MaterialIndex])//if user has overridden materials use it
+			if (MaterialIndex < OverrideMaterials.Num() && OverrideMaterials[MaterialIndex])//if user has overridden materials use it
 			{
-				DestructibleMesh.Materials[NewMaterialIndex].MaterialInterface = OverrideMaterials[MaterialIndex];
-				DestructibleMesh.Materials[NewMaterialIndex].ImportedMaterialSlotName = OverrideMaterials[MaterialIndex]->GetFName();
-				DestructibleMesh.Materials[NewMaterialIndex].MaterialSlotName = OverrideMaterials[MaterialIndex]->GetFName();
+				DestructibleMesh.Materials[MaterialIndex].MaterialInterface = OverrideMaterials[MaterialIndex];
+				DestructibleMesh.Materials[MaterialIndex].ImportedMaterialSlotName = OverrideMaterials[MaterialIndex]->GetFName();
+				DestructibleMesh.Materials[MaterialIndex].MaterialSlotName = OverrideMaterials[MaterialIndex]->GetFName();
 			}
 			else
 			{
-				DestructibleMesh.Materials[NewMaterialIndex].MaterialInterface = DestructibleMesh.FractureSettings->Materials[MaterialIndex];
-				DestructibleMesh.Materials[NewMaterialIndex].ImportedMaterialSlotName = DestructibleMesh.FractureSettings->Materials[MaterialIndex]->GetFName();
-				DestructibleMesh.Materials[NewMaterialIndex].MaterialSlotName = DestructibleMesh.FractureSettings->Materials[MaterialIndex]->GetFName();
+				DestructibleMesh.Materials[MaterialIndex].MaterialInterface = DestructibleMesh.FractureSettings->Materials[MaterialIndex];
+				if (DestructibleMesh.FractureSettings->Materials[MaterialIndex])
+				{
+					DestructibleMesh.Materials[MaterialIndex].ImportedMaterialSlotName = DestructibleMesh.FractureSettings->Materials[MaterialIndex]->GetFName();
+					DestructibleMesh.Materials[MaterialIndex].MaterialSlotName = DestructibleMesh.FractureSettings->Materials[MaterialIndex]->GetFName();
+				}
 			}
-			NewMaterialIndex++;
 		}
 
 		apex::DestructibleAssetCookingDesc DestructibleAssetCookingDesc;
diff --git a/Engine/Plugins/Runtime/ApexDestruction/Source/ApexDestructionEditor/Private/SDestructibleMeshEditorViewport.cpp b/Engine/Plugins/Runtime/ApexDestruction/Source/ApexDestructionEditor/Private/SDestructibleMeshEditorViewport.cpp
--- a/Engine/Plugins/Runtime/ApexDestruction/Source/ApexDestructionEditor/Private/SDestructibleMeshEditorViewport.cpp
+++ b/Engine/Plugins/Runtime/ApexDestruction/Source/ApexDestructionEditor/Private/SDestructibleMeshEditorViewport.cpp
@@ -443,6 +443,38 @@ void FDestructibleMeshEditorViewportClient::ImportFBXChunks()
 			FlushRenderingCommands();
 
 			UnFbx::FFbxImporter* FFbxImporter = UnFbx::FFbxImporter::GetInstance();
+			
+			UnFbx::FBXImportOptions* ImportOptions = FFbxImporter->GetImportOptions();
+			ImportOptions->ImportType = FBXIT_StaticMesh;
+			ImportOptions->bImportScene = false;
+			ImportOptions->bAutoGenerateCollision = false;
+			ImportOptions->bBakePivotInVertex = false;
+			ImportOptions->bCanShowDialog = false;
+			ImportOptions->bCombineToSingle = false;
+			ImportOptions->bConvertScene = true;
+			ImportOptions->bConvertSceneUnit = true;
+			ImportOptions->bCreatePhysicsAsset = false;
+			ImportOptions->bForceFrontXAxis = false;
+			ImportOptions->bGenerateLightmapUVs = false;
+			ImportOptions->bImportLOD = false;
+			ImportOptions->BaseMaterial = nullptr;
+			ImportOptions->bImportRigidMesh = false;
+			ImportOptions->bImportStaticMeshLODs = false;
+			ImportOptions->bInvertNormalMap = false;
+			ImportOptions->bPreserveSmoothingGroups = true;
+			ImportOptions->bAutoComputeLodDistances = false;
+			ImportOptions->bBuildAdjacencyBuffer = true;
+			ImportOptions->bBuildReversedIndexBuffer = true;
+			ImportOptions->bTransformVertexToAbsolute = true;
+			ImportOptions->bUsedAsFullName = true;
+			ImportOptions->ImportRotation = FRotator(0.0f, 0.0f, 0.0f);
+			ImportOptions->ImportTranslation = FVector(0.0f, 0.0f, 0.0f);
+			ImportOptions->ImportUniformScale = 0.0f;
+			ImportOptions->StaticMeshLODGroup = NAME_None;
+			ImportOptions->VertexColorImportOption = EVertexColorImportOption::Replace;
+			ImportOptions->bImportMaterials = true;
+			ImportOptions->bImportTextures = true;
+
 			if (FFbxImporter->ImportFromFile( *ImportFilename, FPaths::GetExtension( ImportFilename ) ) )
 			{
 				TArray<FbxNode*> FbxMeshArray;
