Deprecated DestructibleActor property `bAffectNavigation` to rely fully on the embedded DestructibleComponent. This removes the binding between two properties that may end up out of sync.

diff --git a/Engine/Plugins/Runtime/ApexDestruction/Source/ApexDestruction/Private/DestructibleActor.cpp b/Engine/Plugins/Runtime/ApexDestruction/Source/ApexDestruction/Private/DestructibleActor.cpp
--- a/Engine/Plugins/Runtime/ApexDestruction/Source/ApexDestruction/Private/DestructibleActor.cpp
+++ b/Engine/Plugins/Runtime/ApexDestruction/Source/ApexDestruction/Private/DestructibleActor.cpp
@@ -28,24 +28,17 @@ bool ADestructibleActor::GetReferencedContentObjects( TArray<UObject*>& Objects
 	return true;
 }
 
-void ADestructibleActor::PostEditChangeProperty(FPropertyChangedEvent& PropertyChangedEvent)
-{
-	Super::PostEditChangeProperty(PropertyChangedEvent);
-
-	if (PropertyChangedEvent.Property && PropertyChangedEvent.Property->GetFName() == GET_MEMBER_NAME_CHECKED(ADestructibleActor, bAffectNavigation))
-	{
-		DestructibleComponent->SetCanEverAffectNavigation(bAffectNavigation);
-	}
-}
-#endif // WITH_EDITOR
-
+PRAGMA_DISABLE_DEPRECATION_WARNINGS
 void ADestructibleActor::PostLoad()
 {
 	Super::PostLoad();
-	
-	if (DestructibleComponent)
+
+	if (DestructibleComponent && bAffectNavigation)
 	{
 		DestructibleComponent->SetCanEverAffectNavigation(bAffectNavigation);
+		bAffectNavigation = false;
 	}
 }
+PRAGMA_ENABLE_DEPRECATION_WARNINGS
+#endif // WITH_EDITOR
 
diff --git a/Engine/Plugins/Runtime/ApexDestruction/Source/ApexDestruction/Public/DestructibleActor.h b/Engine/Plugins/Runtime/ApexDestruction/Source/ApexDestruction/Public/DestructibleActor.h
--- a/Engine/Plugins/Runtime/ApexDestruction/Source/ApexDestruction/Public/DestructibleActor.h
+++ b/Engine/Plugins/Runtime/ApexDestruction/Source/ApexDestruction/Public/DestructibleActor.h
@@ -27,21 +27,19 @@ private:
 	UDestructibleComponent* DestructibleComponent;
 public:
 
-	UPROPERTY(config, EditAnywhere, BlueprintReadWrite, Category=Navigation)
-	uint32 bAffectNavigation : 1;
-
 	UPROPERTY(BlueprintAssignable, Category = "Components|Destructible")
 	FActorFractureSignature OnActorFracture;
 
-	//~ Begin AActor Interface.
 #if WITH_EDITOR
-	virtual bool GetReferencedContentObjects( TArray<UObject*>& Objects ) const override;
-	virtual void PostEditChangeProperty(FPropertyChangedEvent& PropertyChangedEvent) override;
-#endif // WITH_EDITOR
-	//~ End AActor Interface.
+	UE_DEPRECATED(4.24, "This property will be removed in future versions. Please use bCanEverAffectNavigation in DestructionComponent.")
+	UPROPERTY(config, BlueprintReadWrite, Category = Navigation, meta = (DeprecationMessage = "Setting the value from Blueprint script does nothing. Please use bCanEverAffectNavigation in DestructionComponent."))
+	uint32 bAffectNavigation : 1;
 
+	//~ Begin AActor Interface.
+	virtual bool GetReferencedContentObjects( TArray<UObject*>& Objects ) const override;
 	virtual void PostLoad() override;
-
+	//~ End AActor Interface.
+#endif // WITH_EDITOR
 
 public:
 	/** Returns DestructibleComponent subobject **/
