Fixing DestructibleActor deprecation by using WITH_EDITORONLY_DATA

#rnx

diff --git a/Engine/Plugins/Runtime/ApexDestruction/Source/ApexDestruction/Public/DestructibleActor.h b/Engine/Plugins/Runtime/ApexDestruction/Source/ApexDestruction/Public/DestructibleActor.h
--- a/Engine/Plugins/Runtime/ApexDestruction/Source/ApexDestruction/Public/DestructibleActor.h
+++ b/Engine/Plugins/Runtime/ApexDestruction/Source/ApexDestruction/Public/DestructibleActor.h
@@ -30,11 +30,13 @@ public:
 	UPROPERTY(BlueprintAssignable, Category = "Components|Destructible")
 	FActorFractureSignature OnActorFracture;
 
-#if WITH_EDITOR
+#if WITH_EDITORONLY_DATA
 	UE_DEPRECATED(4.24, "This property will be removed in future versions. Please use bCanEverAffectNavigation in DestructionComponent.")
 	UPROPERTY(config, BlueprintReadWrite, Category = Navigation, meta = (DeprecationMessage = "Setting the value from Blueprint script does nothing. Please use bCanEverAffectNavigation in DestructionComponent."))
 	uint32 bAffectNavigation : 1;
+#endif // WITH_EDITORONLY_DATA
 
+#if WITH_EDITOR
 	//~ Begin AActor Interface.
 	virtual bool GetReferencedContentObjects( TArray<UObject*>& Objects ) const override;
 	virtual void PostLoad() override;
