Mark HTML5 as blacklisted on ApexDestruction as APEX does not build for HTML5
#rnx

diff --git a/Engine/Plugins/Runtime/ApexDestruction/ApexDestruction.uplugin b/Engine/Plugins/Runtime/ApexDestruction/ApexDestruction.uplugin
--- a/Engine/Plugins/Runtime/ApexDestruction/ApexDestruction.uplugin
+++ b/Engine/Plugins/Runtime/ApexDestruction/ApexDestruction.uplugin
@@ -21,7 +21,8 @@
 			"Type": "Runtime",
 			"LoadingPhase": "PreDefault",
 			"BlacklistPlatforms": [
-				"HoloLens"
+				"HoloLens",
+				"HTML5"
 			]
 		},
     {
