Mark HTML5 as blacklisted on ApexDestruction as APEX does not build for HTML5

#rnx

diff --git a/Engine/Platforms/HTML5/Plugins/Runtime/ApexDestruction/ApexDestruction_HTML5.uplugin b/Engine/Platforms/HTML5/Plugins/Runtime/ApexDestruction/ApexDestruction_HTML5.uplugin
new file mode 100644
--- /dev/null
+++ b/Engine/Platforms/HTML5/Plugins/Runtime/ApexDestruction/ApexDestruction_HTML5.uplugin
@@ -0,0 +1,15 @@
+{
+	"FileVersion" : 3,
+	"bIsPluginExtension": true,
+	"Modules" :
+	[
+		{
+			"Name": "ApexDestruction",
+			"Type": "Runtime",
+			"BlacklistPlatforms": [
+				"HTML5"
+			]
+		}
+	]
+}
+
diff --git a/Engine/Plugins/Runtime/ApexDestruction/ApexDestruction.uplugin b/Engine/Plugins/Runtime/ApexDestruction/ApexDestruction.uplugin
--- a/Engine/Plugins/Runtime/ApexDestruction/ApexDestruction.uplugin
+++ b/Engine/Plugins/Runtime/ApexDestruction/ApexDestruction.uplugin
@@ -21,8 +21,7 @@
 			"Type": "Runtime",
 			"LoadingPhase": "PreDefault",
 			"BlacklistPlatforms": [
-				"HoloLens",
-				"HTML5"
+				"HoloLens"
 			]
 		},
     {
