HTML5 (PE) - disable ApexDestruction

#rnx

diff --git a/Engine/Platforms/HTML5/Plugins/Runtime/ApexDestruction/ApexDestruction_HTML5.uplugin b/Engine/Platforms/HTML5/Plugins/Runtime/ApexDestruction/ApexDestruction_HTML5.uplugin
new file mode 100644
--- /dev/null
+++ b/Engine/Platforms/HTML5/Plugins/Runtime/ApexDestruction/ApexDestruction_HTML5.uplugin
@@ -0,0 +1,15 @@
+{
+	"FileVersion" : 3,
+	"bIsPluginExtension": true,
+	"Modules" :
+	[
+		{
+			"Name": "ApexDestruction",
+			"Type": "Runtime",
+			"BlacklistPlatforms": [
+				"HTML5"
+			]
+		}
+	]
+}
+
