Fixing compile errors from CL 10334624

#rnx

diff --git a/Engine/Plugins/Runtime/ApexDestruction/Source/ApexDestruction/Private/DestructibleComponent.cpp b/Engine/Plugins/Runtime/ApexDestruction/Source/ApexDestruction/Private/DestructibleComponent.cpp
--- a/Engine/Plugins/Runtime/ApexDestruction/Source/ApexDestruction/Private/DestructibleComponent.cpp
+++ b/Engine/Plugins/Runtime/ApexDestruction/Source/ApexDestruction/Private/DestructibleComponent.cpp
@@ -1503,7 +1503,7 @@ void UDestructibleComponent::SetCollisionEnabled(ECollisionEnabled::Type NewType
 #endif // WITH_APEX
 }
 
-void UDestructibleComponent::SetCollisionProfileName(FName InCollisionProfileName)
+void UDestructibleComponent::SetCollisionProfileName(FName InCollisionProfileName, bool bUpdateOverlaps)
 {
     FBodyInstance* LocalInstance = GetBodyInstance();
     if (!LocalInstance)
@@ -1529,7 +1529,7 @@ void UDestructibleComponent::SetCollisionProfileName(FName InCollisionProfileNam
 		{
 			EnsurePhysicsStateCreated();
 		}
-		OnComponentCollisionSettingsChanged();
+		OnComponentCollisionSettingsChanged(bUpdateOverlaps);
 	}
 }
 
diff --git a/Engine/Plugins/Runtime/ApexDestruction/Source/ApexDestruction/Public/DestructibleComponent.h b/Engine/Plugins/Runtime/ApexDestruction/Source/ApexDestruction/Public/DestructibleComponent.h
--- a/Engine/Plugins/Runtime/ApexDestruction/Source/ApexDestruction/Public/DestructibleComponent.h
+++ b/Engine/Plugins/Runtime/ApexDestruction/Source/ApexDestruction/Public/DestructibleComponent.h
@@ -154,7 +154,7 @@ public:
 
 	virtual void SetMaterial(int32 ElementIndex, UMaterialInterface* Material) override;
 	virtual void SetCollisionEnabled(ECollisionEnabled::Type NewType) override;
-    virtual void SetCollisionProfileName(FName InCollisionProfileName) override;
+    virtual void SetCollisionProfileName(FName InCollisionProfileName, bool bUpdateOverlaps=true) override;
 	virtual void OnActorEnableCollisionChanged() override;
 
 	virtual void SetCollisionResponseToChannel(ECollisionChannel Channel, ECollisionResponse NewResponse) override;
diff --git a/Engine/Source/Runtime/Engine/Classes/Components/StaticMeshComponent.h b/Engine/Source/Runtime/Engine/Classes/Components/StaticMeshComponent.h
--- a/Engine/Source/Runtime/Engine/Classes/Components/StaticMeshComponent.h
+++ b/Engine/Source/Runtime/Engine/Classes/Components/StaticMeshComponent.h
@@ -390,7 +390,7 @@ public:
 	UFUNCTION(BlueprintCallable, Category="Rendering|Lighting")
 	void SetReverseCulling(bool ReverseCulling);
 
-	virtual void SetCollisionProfileName(FName InCollisionProfileName) override;
+	virtual void SetCollisionProfileName(FName InCollisionProfileName, bool bUpdateOverlaps=true) override;
 
 public:
 
diff --git a/Engine/Source/Runtime/Engine/Private/Components/StaticMeshComponent.cpp b/Engine/Source/Runtime/Engine/Private/Components/StaticMeshComponent.cpp
--- a/Engine/Source/Runtime/Engine/Private/Components/StaticMeshComponent.cpp
+++ b/Engine/Source/Runtime/Engine/Private/Components/StaticMeshComponent.cpp
@@ -1895,9 +1895,9 @@ void UStaticMeshComponent::GetLocalBounds(FVector& Min, FVector& Max) const
 	}
 }
 
-void UStaticMeshComponent::SetCollisionProfileName(FName InCollisionProfileName)
+void UStaticMeshComponent::SetCollisionProfileName(FName InCollisionProfileName, bool bUpdateOverlaps)
 {
-	Super::SetCollisionProfileName(InCollisionProfileName);
+	Super::SetCollisionProfileName(InCollisionProfileName, bUpdateOverlaps);
 	bUseDefaultCollision = false;
 }
 
diff --git a/Engine/Source/Runtime/MRMesh/Private/MRMeshComponent.cpp b/Engine/Source/Runtime/MRMesh/Private/MRMeshComponent.cpp
--- a/Engine/Source/Runtime/MRMesh/Private/MRMeshComponent.cpp
+++ b/Engine/Source/Runtime/MRMesh/Private/MRMeshComponent.cpp
@@ -556,7 +556,7 @@ void UMRMeshComponent::SetCollisionEnabled(ECollisionEnabled::Type NewType)
 	Super::SetCollisionEnabled(NewType);
 }
 
-void UMRMeshComponent::SetCollisionProfileName(FName InCollisionProfileName)
+void UMRMeshComponent::SetCollisionProfileName(FName InCollisionProfileName, bool bUpdateOverlaps)
 {
 	SCOPE_CYCLE_COUNTER(STAT_MrMesh_SetCollisionProfileName);
 
@@ -577,7 +577,7 @@ void UMRMeshComponent::SetCollisionProfileName(FName InCollisionProfileName)
 		}
 	}
 
-	Super::SetCollisionProfileName(InCollisionProfileName);
+	Super::SetCollisionProfileName(InCollisionProfileName, bUpdateOverlaps);
 }
 
 void UMRMeshComponent::SetCollisionObjectType(ECollisionChannel Channel)
diff --git a/Engine/Source/Runtime/MRMesh/Public/MRMeshComponent.h b/Engine/Source/Runtime/MRMesh/Public/MRMeshComponent.h
--- a/Engine/Source/Runtime/MRMesh/Public/MRMeshComponent.h
+++ b/Engine/Source/Runtime/MRMesh/Public/MRMeshComponent.h
@@ -106,7 +106,7 @@ public:
 	virtual void SetCollisionResponseToAllChannels(ECollisionResponse NewResponse) override;
 	virtual void SetCollisionResponseToChannels(const FCollisionResponseContainer& NewResponses) override;
 	virtual void SetCollisionEnabled(ECollisionEnabled::Type NewType) override;
-	virtual void SetCollisionProfileName(FName InCollisionProfileName) override;
+	virtual void SetCollisionProfileName(FName InCollisionProfileName, bool bUpdateOverlaps=true) override;
 
 	virtual void SetWalkableSlopeOverride(const FWalkableSlopeOverride& NewOverride) override;
 
