Fix static analysis warning: Expression 'DepthCount > 0' is always false
#rnx

diff --git a/Engine/Plugins/Runtime/ApexDestruction/Source/ApexDestructionEditor/Private/SDestructibleMeshEditorViewport.cpp b/Engine/Plugins/Runtime/ApexDestruction/Source/ApexDestructionEditor/Private/SDestructibleMeshEditorViewport.cpp
--- a/Engine/Plugins/Runtime/ApexDestruction/Source/ApexDestructionEditor/Private/SDestructibleMeshEditorViewport.cpp
+++ b/Engine/Plugins/Runtime/ApexDestruction/Source/ApexDestructionEditor/Private/SDestructibleMeshEditorViewport.cpp
@@ -663,19 +663,19 @@ bool SDestructibleMeshEditorViewport::IsVisible() const
 void SDestructibleMeshEditorViewport::SetPreviewDepth(uint32 InPreviewDepth)
 {
 	uint32 DepthCount = 0;
+	uint32 NewPreviewDepth = 0;
 
 #if WITH_APEX
-	if (DestructibleMesh != NULL && DestructibleMesh->ApexDestructibleAsset != NULL)
+	if (DestructibleMesh && DestructibleMesh->ApexDestructibleAsset)
 	{
 		DepthCount = DestructibleMesh->ApexDestructibleAsset->getDepthCount();
-	}
-#endif // WITH_APEX
 
-	uint32 NewPreviewDepth = 0;
-	if (DepthCount > 0)
-	{
-		NewPreviewDepth = FMath::Clamp(InPreviewDepth, (uint32)0, DepthCount-1);
+		if (DepthCount > 0)
+		{
+			NewPreviewDepth = FMath::Clamp(InPreviewDepth, (uint32)0, DepthCount-1);
+		}
 	}
+#endif // WITH_APEX
 
 	if (NewPreviewDepth != PreviewDepth)
 	{
