Fix destructible mesh, which was always using the same derived data key.

diff --git a/Engine/Plugins/Runtime/ApexDestruction/Source/ApexDestruction/Private/ApexDestructibleAssetImport.cpp b/Engine/Plugins/Runtime/ApexDestruction/Source/ApexDestruction/Private/ApexDestructibleAssetImport.cpp
--- a/Engine/Plugins/Runtime/ApexDestruction/Source/ApexDestruction/Private/ApexDestructibleAssetImport.cpp
+++ b/Engine/Plugins/Runtime/ApexDestruction/Source/ApexDestruction/Private/ApexDestructibleAssetImport.cpp
@@ -766,6 +766,8 @@ bool SetApexDestructibleAsset(UDestructibleMesh& DestructibleMesh, apex::Destruc
 {
 	DestructibleMesh.PreEditChange(NULL);
 
+	DestructibleMesh.InvalidateDeriveDataCacheGUID();
+
 	ExistingDestMeshData * ExistDestMeshDataPtr = nullptr;
 	if(Options & EDestructibleImportOptions::PreserveSettings)
 	{
