Fixed warning

diff --git a/Engine/Plugins/Runtime/ApexDestruction/Source/ApexDestructionEditor/Private/DestructibleMeshThumbnailRenderer.cpp b/Engine/Plugins/Runtime/ApexDestruction/Source/ApexDestructionEditor/Private/DestructibleMeshThumbnailRenderer.cpp
--- a/Engine/Plugins/Runtime/ApexDestruction/Source/ApexDestructionEditor/Private/DestructibleMeshThumbnailRenderer.cpp
+++ b/Engine/Plugins/Runtime/ApexDestruction/Source/ApexDestructionEditor/Private/DestructibleMeshThumbnailRenderer.cpp
@@ -16,7 +16,7 @@ UDestructibleMeshThumbnailRenderer::UDestructibleMeshThumbnailRenderer(const FOb
 	ThumbnailScene = nullptr;
 }
 
-void UDestructibleMeshThumbnailRenderer::Draw(UObject* Object, int32 X, int32 Y, uint32 Width, uint32 Height, FRenderTarget* RenderTarget, FCanvas* Canvas)
+void UDestructibleMeshThumbnailRenderer::Draw(UObject* Object, int32 X, int32 Y, uint32 Width, uint32 Height, FRenderTarget* RenderTarget, FCanvas* Canvas, bool bAdditionalViewFamily)
 {
 	UDestructibleMesh* DestructibleMesh = Cast<UDestructibleMesh>(Object);
 	if (DestructibleMesh != nullptr)
@@ -28,7 +28,8 @@ void UDestructibleMeshThumbnailRenderer::Draw(UObject* Object, int32 X, int32 Y,
 
 		ThumbnailScene->SetDestructibleMesh(DestructibleMesh);
 		FSceneViewFamilyContext ViewFamily( FSceneViewFamily::ConstructionValues( RenderTarget, ThumbnailScene->GetScene(), FEngineShowFlags(ESFIM_Game) )
-			.SetWorldTimes(FApp::GetCurrentTime() - GStartTime, FApp::GetDeltaTime(), FApp::GetCurrentTime() - GStartTime));
+			.SetWorldTimes(FApp::GetCurrentTime() - GStartTime, FApp::GetDeltaTime(), FApp::GetCurrentTime() - GStartTime)
+			.SetAdditionalViewFamily(bAdditionalViewFamily));
 
 		ViewFamily.EngineShowFlags.DisableAdvancedFeatures();
 		ViewFamily.EngineShowFlags.MotionBlur = 0;
diff --git a/Engine/Plugins/Runtime/ApexDestruction/Source/ApexDestructionEditor/Private/DestructibleMeshThumbnailRenderer.h b/Engine/Plugins/Runtime/ApexDestruction/Source/ApexDestructionEditor/Private/DestructibleMeshThumbnailRenderer.h
--- a/Engine/Plugins/Runtime/ApexDestruction/Source/ApexDestructionEditor/Private/DestructibleMeshThumbnailRenderer.h
+++ b/Engine/Plugins/Runtime/ApexDestruction/Source/ApexDestructionEditor/Private/DestructibleMeshThumbnailRenderer.h
@@ -22,7 +22,7 @@ class UDestructibleMeshThumbnailRenderer : public UDefaultSizedThumbnailRenderer
 
 
 	// Begin UThumbnailRenderer Object
-	APEXDESTRUCTIONEDITOR_API virtual void Draw(UObject* Object, int32 X, int32 Y, uint32 Width, uint32 Height, FRenderTarget*, FCanvas* Canvas) override;
+	APEXDESTRUCTIONEDITOR_API virtual void Draw(UObject* Object, int32 X, int32 Y, uint32 Width, uint32 Height, FRenderTarget*, FCanvas* Canvas, bool bAdditionalViewFamily) override;
 	// End UThumbnailRenderer Object
 
 	// UObject implementation
