Removed the context menu option to create a destructible mesh through APEX when APEX has not been built for the current engine config. Also early out of creation process if we enter it directly.

diff --git a/Engine/Plugins/Runtime/ApexDestruction/Source/ApexDestructionEditor/Private/ApexDestructionEditorModule.cpp b/Engine/Plugins/Runtime/ApexDestruction/Source/ApexDestructionEditor/Private/ApexDestructionEditorModule.cpp
--- a/Engine/Plugins/Runtime/ApexDestruction/Source/ApexDestructionEditor/Private/ApexDestructionEditorModule.cpp
+++ b/Engine/Plugins/Runtime/ApexDestruction/Source/ApexDestructionEditor/Private/ApexDestructionEditorModule.cpp
@@ -71,8 +71,10 @@ PRAGMA_ENABLE_DEPRECATION_WARNINGS
 
 TSharedRef<FExtender> FDestructibleMeshEditorModule::OnExtendContentBrowserAssetSelectionMenu(const TArray<FAssetData>& SelectedAssets)
 {
+
 	TSharedRef<FExtender> Extender(new FExtender());
 
+#if WITH_APEX
 	// Run thru the assets to determine if any meet our criteria
 	bool bAnyStaticMesh = false;
 	for (auto AssetIt = SelectedAssets.CreateConstIterator(); AssetIt; ++AssetIt)
@@ -93,6 +95,7 @@ TSharedRef<FExtender> FDestructibleMeshEditorModule::OnExtendContentBrowserAsset
 				FUIAction(FExecuteAction::CreateStatic(&FAssetTypeActions_DestructibleMesh::ExecuteCreateDestructibleMeshes, SelectedAssets), FCanExecuteAction()));
 		}));
 	}
+#endif
 
 	return Extender;
 }
@@ -141,6 +144,11 @@ PRAGMA_ENABLE_DEPRECATION_WARNINGS
 PRAGMA_DISABLE_DEPRECATION_WARNINGS
 UDestructibleMesh* FDestructibleMeshEditorModule::CreateDestructibleMeshFromStaticMesh(UObject* InParent, UStaticMesh* StaticMesh, FName Name, EObjectFlags Flags, FText& OutErrorMsg)
 {
+#if !WITH_APEX
+	OutErrorMsg = LOCTEXT("NoAPEX", "APEX is not built for the current engine configuration.");
+	return nullptr;
+#endif
+
 	if (StaticMesh == NULL)
 	{
 		OutErrorMsg = LOCTEXT( "StaticMeshInvalid", "Static Mesh is Invalid!" );
