Fixed local bounds calculations for destructible components.
 - getBounds() on the APEX actor gets world bounds, CalcBounds expects local concatenated with the provided transform
 - getLocalBounds() on the APEX actor does not update after fracture
 - getChunkLocalBounds on the APEX actor does not update after fracture
 - Settled on taking the world bounds, inverting the actor transform then applying the supplied transform to create reasonable bounds.

diff --git a/Engine/Plugins/Runtime/ApexDestruction/Source/ApexDestruction/Private/DestructibleComponent.cpp b/Engine/Plugins/Runtime/ApexDestruction/Source/ApexDestruction/Private/DestructibleComponent.cpp
--- a/Engine/Plugins/Runtime/ApexDestruction/Source/ApexDestruction/Private/DestructibleComponent.cpp
+++ b/Engine/Plugins/Runtime/ApexDestruction/Source/ApexDestruction/Private/DestructibleComponent.cpp
@@ -101,9 +101,14 @@ FBoxSphereBounds UDestructibleComponent::CalcBounds(const FTransform& LocalToWor
 		return Super::CalcBounds(LocalToWorld);
 	}
 
-	const PxBounds3& PBounds = ApexDestructibleActor->getBounds();
-
-	return FBoxSphereBounds(FBox(P2UVector(PBounds.minimum), P2UVector(PBounds.maximum)));
+	// getLocalBounds and getChunkLocalBounds are non-updating in the APEX actor, so using those will not
+	// get us the correct bounds. Instead we take the world bounds and recalculate the local bounds
+	// from it.
+	const PxBounds3 PWorldBounds = ApexDestructibleActor->getBounds();
+	
+	return FBoxSphereBounds(FBox(P2UVector(PWorldBounds.minimum), P2UVector(PWorldBounds.maximum)))
+		.TransformBy(GetComponentToWorld().Inverse())
+		.TransformBy(LocalToWorld);
 
 #else	// #if WITH_APEX
 	return Super::CalcBounds(LocalToWorld);
