Fix Apex build breakage

#rnx

diff --git a/Engine/Plugins/Runtime/ApexDestruction/Source/ApexDestruction/Private/DestructibleComponent.cpp b/Engine/Plugins/Runtime/ApexDestruction/Source/ApexDestruction/Private/DestructibleComponent.cpp
--- a/Engine/Plugins/Runtime/ApexDestruction/Source/ApexDestruction/Private/DestructibleComponent.cpp
+++ b/Engine/Plugins/Runtime/ApexDestruction/Source/ApexDestruction/Private/DestructibleComponent.cpp
@@ -482,12 +482,12 @@ void UDestructibleComponent::AddImpulse( FVector Impulse, FName BoneName /*= NAM
 	ExecuteOnPhysicsReadWrite([&]
 	{
 		const int32 ChunkIdx = BoneIdxToChunkIdx(GetBoneIndex(BoneName));
-		ApexDestructibleActor->addForce(ChunkIdx, U2PVector(Impulse),  bVelChange ? PxForceMode::eVELOCITY_CHANGE : PxForceMode::eIMPULSE);
+		ApexDestructibleActor->addForce(ChunkIdx, U2PVector(Impulse), bVelChange ? PxForceMode::eVELOCITY_CHANGE : PxForceMode::eIMPULSE);
 	});
 #endif
 }
 
-void UDestructibleComponent::AddImpulseAtLocation( FVector Impulse, FVector Position, FName BoneName /*= NAME_None*/ )
+void UDestructibleComponent::AddImpulseAtLocation( FVector Impulse, FVector Position, FName BoneName /*= NAME_None*/, bool bVelChange /*= false*/ )
 {
 #if WITH_APEX
 	ExecuteOnPhysicsReadWrite([&]
@@ -495,7 +495,7 @@ void UDestructibleComponent::AddImpulseAtLocation( FVector Impulse, FVector Posi
 		const int32 ChunkIdx = BoneIdxToChunkIdx(GetBoneIndex(BoneName));
 		PxVec3 PxPosition = U2PVector(Position);
 
-		ApexDestructibleActor->addForce(ChunkIdx, U2PVector(Impulse),  PxForceMode::eIMPULSE, &PxPosition);
+		ApexDestructibleActor->addForce(ChunkIdx, U2PVector(Impulse), bVelChange ? PxForceMode::eVELOCITY_CHANGE : PxForceMode::eIMPULSE, &PxPosition);
 	});
 #endif
 }
diff --git a/Engine/Plugins/Runtime/ApexDestruction/Source/ApexDestruction/Public/DestructibleComponent.h b/Engine/Plugins/Runtime/ApexDestruction/Source/ApexDestruction/Public/DestructibleComponent.h
--- a/Engine/Plugins/Runtime/ApexDestruction/Source/ApexDestruction/Public/DestructibleComponent.h
+++ b/Engine/Plugins/Runtime/ApexDestruction/Source/ApexDestruction/Public/DestructibleComponent.h
@@ -140,7 +140,7 @@ public:
 	virtual bool IsAnySimulatingPhysics() const override;
 
 	virtual void AddImpulse(FVector Impulse, FName BoneName = NAME_None, bool bVelChange = false) override;
-	virtual void AddImpulseAtLocation(FVector Impulse, FVector Position, FName BoneName = NAME_None) override;
+	virtual void AddImpulseAtLocation(FVector Impulse, FVector Position, FName BoneName = NAME_None, bool bVelChange = false) override;
 	virtual void AddForce(FVector Force, FName BoneName = NAME_None, bool bAccelChange = false) override;
 	virtual void AddForceAtLocation(FVector Force, FVector Location, FName BoneName = NAME_None) override;
 	virtual void AddForceAtLocationLocal(FVector Force, FVector Location, FName BoneName = NAME_None) override;
