Restore destructible mesh classes to avoid loading errors for legacy assets.

diff --git a/Engine/Plugins/Runtime/ApexDestruction/ApexDestruction.uplugin b/Engine/Plugins/Runtime/ApexDestruction/ApexDestruction.uplugin
new file mode 100644
--- /dev/null
+++ b/Engine/Plugins/Runtime/ApexDestruction/ApexDestruction.uplugin
@@ -0,0 +1,33 @@
+{
+	"FileVersion" : 3,
+	"Version" : 1,
+	"VersionName" : "1.0",
+	"FriendlyName" : "Apex Destruction",
+	"Description" : "APEX implementation of destruction",
+	"Category" : "Physics",
+	"CreatedBy" : "Epic Games, Inc.",
+	"CreatedByURL" : "http://epicgames.com",
+	"DocsURL" : "",
+	"MarketplaceURL" : "",
+	"SupportURL" : "",
+	"EnabledByDefault" : false,
+	"CanContainContent" : false,
+	"IsBetaVersion" : false,
+	"Installed" : false,
+	"Modules" :
+	[
+		{
+			"Name": "ApexDestruction",
+			"Type": "Runtime",
+			"LoadingPhase": "PreDefault",
+			"PlatformDenyList": [
+				"HoloLens"
+			]
+		},
+		{
+			"Name": "ApexDestructionEditor",
+			"Type": "UncookedOnly",
+			"LoadingPhase": "PreDefault"
+		}
+	]
+}
diff --git a/Engine/Plugins/Runtime/ApexDestruction/Config/BaseApexDestruction.ini b/Engine/Plugins/Runtime/ApexDestruction/Config/BaseApexDestruction.ini
new file mode 100644
--- /dev/null
+++ b/Engine/Plugins/Runtime/ApexDestruction/Config/BaseApexDestruction.ini
@@ -0,0 +1,6 @@
+[CoreRedirects]
+
++ClassRedirects=(OldName="DestructibleActor", NewName="/Script/ApexDestruction.DestructibleActor")
++ClassRedirects=(OldName="DestructibleComponent", NewName="/Script/ApexDestruction.DestructibleComponent")
++ClassRedirects=(OldName="DestructibleMesh", NewName="/Script/ApexDestruction.DestructibleMesh")
++ClassRedirects=(OldName="DestructibleFractureSettings", NewName="/Script/ApexDestruction.DestructibleFractureSettings")
\ No newline at end of file
diff --git a/Engine/Plugins/Runtime/ApexDestruction/Source/ApexDestruction/ApexDestruction.Build.cs b/Engine/Plugins/Runtime/ApexDestruction/Source/ApexDestruction/ApexDestruction.Build.cs
new file mode 100644
--- /dev/null
+++ b/Engine/Plugins/Runtime/ApexDestruction/Source/ApexDestruction/ApexDestruction.Build.cs
@@ -0,0 +1,36 @@
+// Copyright Epic Games, Inc. All Rights Reserved.
+
+using UnrealBuildTool;
+using System.IO;
+
+public class ApexDestruction : ModuleRules
+{
+	public ApexDestruction(ReadOnlyTargetRules Target) : base(Target)
+	{
+		PrivateDependencyModuleNames.AddRange(
+			new string[] {
+				"Core",
+				"CoreUObject", // @todo Mac: for some reason it's needed to link in debug on Mac
+				"Engine",
+//				"PhysX",
+//				"APEX",
+				"RHI",
+				"RenderCore",
+//				"ApexDestructionLib",
+				"NavigationSystem",
+            }
+		);
+
+        if (Target.Type == TargetType.Editor)
+        {
+            PrivateDependencyModuleNames.AddRange(
+                new string[] {
+					"EditorFramework",
+                    "UnrealEd"
+                }
+            );
+        }
+        
+        SetupModulePhysicsSupport(Target);
+	}
+}
diff --git a/Engine/Plugins/Runtime/ApexDestruction/Source/ApexDestruction/Private/ApexDestructibleAssetImport.cpp b/Engine/Plugins/Runtime/ApexDestruction/Source/ApexDestruction/Private/ApexDestructibleAssetImport.cpp
new file mode 100644
--- /dev/null
+++ b/Engine/Plugins/Runtime/ApexDestruction/Source/ApexDestruction/Private/ApexDestructibleAssetImport.cpp
@@ -0,0 +1,1110 @@
+// Copyright Epic Games, Inc. All Rights Reserved.
+
+/*=============================================================================
+	ApexDestructibleAssetImport.cpp:
+
+	Creation of an APEX NxDestructibleAsset from a binary buffer.
+	SkeletalMesh creation from an APEX NxDestructibleAsset.
+
+	SkeletalMesh creation largely based on FbxSkeletalMeshImport.cpp
+
+=============================================================================*/
+
+#include "ApexDestructibleAssetImport.h"
+
+#if WITH_EDITOR
+#include "Modules/ModuleManager.h"
+#include "Engine/SkeletalMesh.h"
+#include "Materials/Material.h"
+#include "Factories/Factory.h"
+#include "EditorFramework/AssetImportData.h"
+#include "DestructibleComponent.h"
+#include "ApexDestructionModule.h"
+
+DEFINE_LOG_CATEGORY_STATIC(LogApexDestructibleAssetImport, Log, All);
+
+#include "PhysicsPublic.h"
+#include "ImportUtils/SkeletalMeshImportUtils.h"
+#include "PhysXIncludes.h"
+#include "MeshUtilities.h"
+#include "UObject/UObjectHash.h"
+#include "UObject/UObjectIterator.h"
+#include "ComponentReregisterContext.h"
+#include "Rendering/SkeletalMeshModel.h"
+#include "DestructibleMesh.h"
+#include "Factories/FbxSkeletalMeshImportData.h"
+#include "Interfaces/ITargetPlatform.h"
+#include "Interfaces/ITargetPlatformManagerModule.h"
+#include "Misc/CoreMisc.h"
+
+PRAGMA_DISABLE_DEPRECATION_WARNINGS
+
+#if WITH_APEX
+
+
+// Set to 1 to match UE transformation of FBX-imported meshes
+#define INVERT_Y_AND_V	1
+
+// Set to 1 to use temporary transform function
+#define USE_TEMPORARY_TRANSFORMATION_FUNCTION 1
+
+
+// Forward declarations and external processing functions
+struct ExistingSkelMeshData;
+
+
+// Temporary transform function, to be removed once the APEX SDK is updated
+#if USE_TEMPORARY_TRANSFORMATION_FUNCTION
+
+static void ApplyTransformationToApexDestructibleAsset( apex::DestructibleAsset& ApexDestructibleAsset, const physx::PxMat44& Transform )
+{
+	// Get the NvParameterized interface to the asset
+	NvParameterized::Interface* AssetParams = const_cast<NvParameterized::Interface*>(ApexDestructibleAsset.getAssetNvParameterized());
+	if (AssetParams != NULL)
+	{
+		// Name buffer used for parameterized array element lookup
+		char ArrayElementName[1024];
+
+		/* surfaceTrace default normal */
+		physx::PxI32 SurfaceTraceSetCount;
+		if (NvParameterized::getParamArraySize(*AssetParams, "surfaceTraceSets", SurfaceTraceSetCount))
+		{
+			for (physx::PxI32 i = 0; i < SurfaceTraceSetCount; ++i)
+			{
+				FCStringAnsi::Sprintf( ArrayElementName, "surfaceTraceSets[%d]", i );
+				NvParameterized::Interface* SurfaceTraceSetParams;
+				if (NvParameterized::getParamRef(*AssetParams, ArrayElementName, SurfaceTraceSetParams))
+				{
+					if (SurfaceTraceSetParams != NULL)
+					{
+						physx::PxI32 SurfaceTraceCount;
+						if (NvParameterized::getParamArraySize(*SurfaceTraceSetParams, "traces", SurfaceTraceCount))
+						{
+							for (physx::PxI32 j = 0; j < SurfaceTraceCount; ++j)
+							{
+								FCStringAnsi::Sprintf( ArrayElementName, "traces[%d].defaultNormal", j );
+								NvParameterized::Handle DefaultNormalHandle(*SurfaceTraceSetParams);
+								if (NvParameterized::findParam(*SurfaceTraceSetParams, ArrayElementName, DefaultNormalHandle) != NULL)
+								{
+									physx::PxVec3 DefaultNormal;
+									DefaultNormalHandle.getParamVec3( DefaultNormal );
+									DefaultNormal = Transform.rotate( DefaultNormal );
+									DefaultNormalHandle.setParamVec3( DefaultNormal );
+								}
+							}
+						}
+					}
+				}
+			}
+		}
+
+		/* For now, we'll just clear the current cached streams. */
+		NvParameterized::Interface* CollisionDataParams;
+		if (NvParameterized::getParamRef(*AssetParams, "collisionData", CollisionDataParams))
+		{
+			if (CollisionDataParams != NULL)
+			{
+				CollisionDataParams->destroy();
+				NvParameterized::setParamRef(*AssetParams, "collisionData", NULL);
+			}
+		}
+
+		/* chunk surface normal */
+		physx::PxI32 AssetChunkCount;
+		if (NvParameterized::getParamArraySize(*AssetParams, "chunks", AssetChunkCount))
+		{
+			for (physx::PxI32 i = 0; i < AssetChunkCount; ++i)
+			{
+				FCStringAnsi::Sprintf( ArrayElementName, "chunks[%d].surfaceNormal", i );
+				NvParameterized::Handle ChunkSurfaceNormalHandle(*AssetParams);
+				if (NvParameterized::findParam(*AssetParams, ArrayElementName, ChunkSurfaceNormalHandle) != NULL)
+				{
+					physx::PxVec3 ChunkSurfaceNormal;
+					ChunkSurfaceNormalHandle.getParamVec3( ChunkSurfaceNormal );
+					ChunkSurfaceNormal = Transform.rotate( ChunkSurfaceNormal );
+					ChunkSurfaceNormalHandle.setParamVec3( ChunkSurfaceNormal );
+				}
+			}
+		}
+
+		/* bounds */
+		physx::PxBounds3 Bounds;
+		if (NvParameterized::getParamBounds3(*AssetParams, "bounds", Bounds))
+		{
+			if (!Bounds.isEmpty())
+			{
+				Bounds = physx::PxBounds3::basisExtent(Transform.transform(Bounds.getCenter()), physx::PxMat33(Transform.column0.getXYZ(), Transform.column1.getXYZ(), Transform.column2.getXYZ()), Bounds.getExtents());
+				NvParameterized::setParamBounds3(*AssetParams, "bounds", Bounds);
+			}
+		}
+
+		/* chunk convex hulls */
+		physx::PxI32 ConvexHullCount;
+		if (NvParameterized::getParamArraySize(*AssetParams, "chunkConvexHulls", ConvexHullCount))
+		{
+			for (physx::PxI32 i = 0; i < ConvexHullCount; ++i)
+			{
+				FCStringAnsi::Sprintf( ArrayElementName, "chunkConvexHulls[%d]", i );
+				NvParameterized::Interface* ConvexHullParams;
+				if (NvParameterized::getParamRef(*AssetParams, ArrayElementName, ConvexHullParams))
+				{
+					if (ConvexHullParams != NULL)
+					{
+						// planes
+						physx::PxI32 UniquePlaneCount;
+						if (NvParameterized::getParamArraySize(*ConvexHullParams, "uniquePlanes", UniquePlaneCount))
+						{
+							for (physx::PxI32 j = 0; j < UniquePlaneCount; ++j)
+							{
+								FCStringAnsi::Sprintf( ArrayElementName, "uniquePlanes[%d].normal", j );
+								NvParameterized::Handle PlaneNormalHandle(*ConvexHullParams);
+								if (NvParameterized::findParam(*ConvexHullParams, ArrayElementName, PlaneNormalHandle) != NULL)
+								{
+									physx::PxVec3 PlaneNormal;
+									PlaneNormalHandle.getParamVec3( PlaneNormal );
+									PlaneNormal = Transform.rotate( PlaneNormal );
+									PlaneNormalHandle.setParamVec3( PlaneNormal );
+								}
+							}
+						}
+
+						// vertices
+						physx::PxBounds3 HullBounds;
+						HullBounds.setEmpty();
+						physx::PxI32 HullVertexCount;
+						if (NvParameterized::getParamArraySize(*ConvexHullParams, "vertices", HullVertexCount))
+						{
+							for (physx::PxI32 j = 0; j < HullVertexCount; ++j)
+							{
+								FCStringAnsi::Sprintf( ArrayElementName, "vertices[%d]", j );
+								NvParameterized::Handle HullVertexHandle(*ConvexHullParams);
+								if (NvParameterized::findParam(*ConvexHullParams, ArrayElementName, HullVertexHandle) != NULL)
+								{
+									physx::PxVec3 HullVertex;
+									HullVertexHandle.getParamVec3( HullVertex );
+									HullVertex = Transform.transform( HullVertex );
+									HullVertexHandle.setParamVec3( HullVertex );
+									HullBounds.include( HullVertex );
+								}
+							}
+						}
+						NvParameterized::setParamBounds3(*ConvexHullParams, "bounds", HullBounds);
+					}
+				}
+			}
+		}
+
+		/* render mesh asset (bounding boxes only) */
+		const physx::PxMat33 Basis(Transform.getBasis(0), Transform.getBasis(1), Transform.getBasis(2));
+		const physx::PxVec3 Offset = Transform.getPosition();
+		NvParameterized::Interface* RenderMeshAssetParams;
+		if (NvParameterized::getParamRef(*AssetParams, "renderMeshAsset", RenderMeshAssetParams))
+		{
+			physx::PxI32 PartBoundsCount;
+			if (NvParameterized::getParamArraySize(*RenderMeshAssetParams, "partBounds", PartBoundsCount))
+			{
+				for (physx::PxI32 i = 0; i < PartBoundsCount; ++i)
+				{
+					FCStringAnsi::Sprintf( ArrayElementName, "partBounds[%d]", i );
+					NvParameterized::Handle PartBoundsHandle(*RenderMeshAssetParams);
+					if (NvParameterized::findParam(*RenderMeshAssetParams, ArrayElementName, PartBoundsHandle) != NULL)
+					{
+						physx::PxBounds3 PartBounds;
+						PartBoundsHandle.getParamBounds3( PartBounds );
+						PartBounds = physx::PxBounds3::transformSafe(Basis, PartBounds);
+						PartBounds.minimum += Offset;
+						PartBounds.maximum += Offset;
+						PartBoundsHandle.setParamBounds3( PartBounds );
+					}
+				}
+			}
+		}
+	}
+}
+#endif // #if USE_TEMPORARY_TRANSFORMATION_FUNCTION
+
+// Storage for destructible mesh settings (including base skeletal mesh)
+struct FExistingDestMeshData
+{
+	FExistingDestMeshData() : BodySetup(NULL) {}
+
+	TSharedPtr<FExistingSkelMeshData>	SkelMeshData;
+	UBodySetup*							BodySetup;
+	TArray<struct FFractureEffect>		FractureEffects;
+};
+
+TSharedPtr<FExistingDestMeshData> SaveExistingDestMeshData(UDestructibleMesh* ExistingDestructibleMesh)
+{
+	TSharedPtr<FExistingDestMeshData> ExistingDestMeshDataPtr;
+
+	if (ExistingDestructibleMesh)
+	{
+		ExistingDestMeshDataPtr = MakeShared<FExistingDestMeshData>();
+
+		// Only save off SkelMeshData if it's been created
+		ExistingDestMeshDataPtr->SkelMeshData = SkeletalMeshImportUtils::SaveExistingSkelMeshData(ExistingDestructibleMesh, true, INDEX_NONE);
+		ExistingDestMeshDataPtr->BodySetup = ExistingDestructibleMesh->BodySetup;
+		ExistingDestMeshDataPtr->FractureEffects = ExistingDestructibleMesh->FractureEffects;
+	}
+
+	return ExistingDestMeshDataPtr;
+}
+
+static void RestoreExistingDestMeshData(const TSharedPtr<FExistingDestMeshData>& MeshData, UDestructibleMesh* DestructibleMesh)
+{
+	if (MeshData && DestructibleMesh)
+	{
+		// Restore old settings, but resize arrays to make sense with the new NxDestructibleAsset
+		if (MeshData->SkelMeshData)
+		{
+			SkeletalMeshImportUtils::RestoreExistingSkelMeshData(MeshData->SkelMeshData, DestructibleMesh, INDEX_NONE, false, false, false);
+		}
+		DestructibleMesh->BodySetup =  MeshData->BodySetup;
+		DestructibleMesh->FractureEffects = MeshData->FractureEffects;
+
+		// Resize the depth parameters array to the appropriate size
+		const int32 ParamDepthDifference = (int32)DestructibleMesh->ApexDestructibleAsset->getDepthCount() - DestructibleMesh->DefaultDestructibleParameters.DepthParameters.Num();
+		if (ParamDepthDifference > 0)
+		{
+			for (int32 i = 0; i < ParamDepthDifference; ++i)
+			{
+				DestructibleMesh->DefaultDestructibleParameters.DepthParameters.Add(FDestructibleDepthParameters());
+			}
+		}
+		else
+		if (ParamDepthDifference < 0)
+		{
+			DestructibleMesh->DefaultDestructibleParameters.DepthParameters.RemoveAt(DestructibleMesh->DefaultDestructibleParameters.DepthParameters.Num()+ParamDepthDifference, -ParamDepthDifference);
+		}
+
+		// Resize the fracture effects array to the appropriate size
+		const int32 EffectsDepthDifference = (int32)DestructibleMesh->ApexDestructibleAsset->getDepthCount() - DestructibleMesh->FractureEffects.Num();
+		if (EffectsDepthDifference > 0)
+		{
+			DestructibleMesh->FractureEffects.AddZeroed(EffectsDepthDifference);
+		}
+		else
+		if (EffectsDepthDifference < 0)
+		{
+			DestructibleMesh->FractureEffects.RemoveAt(DestructibleMesh->FractureEffects.Num()+EffectsDepthDifference, -EffectsDepthDifference);
+		}
+	}
+}
+
+/**
+ * Extract the material names from the Apex Render Mesh contained within an Apex Destructible Asset.
+ * @param ImportData - SkeletalMesh import data into which we are extracting information
+ * @param ApexDestructibleAsset - the Apex Destructible Asset
+ */
+static void ImportMaterialsForSkelMesh(FSkeletalMeshImportData &ImportData, const apex::DestructibleAsset& ApexDestructibleAsset)
+{
+	physx::PxU32 SubmeshCount = 0;
+
+	// Get the submesh count from the Destructible Asset's Render Mesh
+	const apex::RenderMeshAsset* ApexRenderMesh = ApexDestructibleAsset.getRenderMeshAsset();
+	if (ApexRenderMesh != NULL)
+	{
+		SubmeshCount = ApexRenderMesh->getSubmeshCount();
+	}
+
+	if( SubmeshCount == 0 )
+	{
+		// No material info, create a default material slot
+		++SubmeshCount;
+		UE_LOG(LogApexDestructibleAssetImport, Warning,TEXT("No material associated with skeletal mesh - using default"));
+	}
+	else
+	{
+		UE_LOG(LogApexDestructibleAssetImport, Warning,TEXT("Using default materials for material slot"));
+	}
+
+	// Create material slots
+	UMaterial* DefaultMaterial = UMaterial::GetDefaultMaterial(MD_Surface);
+	if (DefaultMaterial && ApexRenderMesh != nullptr)
+	{
+		TArray<FString> Materials;
+		for (uint32 SubmeshIdx = 0; SubmeshIdx < SubmeshCount; SubmeshIdx++)
+		{
+			FString MaterialName = ApexRenderMesh->getMaterialName(SubmeshIdx);
+			if (!Materials.Contains(MaterialName))
+			{
+				Materials.Add(MaterialName);
+				ImportData.Materials.Add(SkeletalMeshImportData::FMaterial());
+				ImportData.Materials.Last().Material = DefaultMaterial;
+				ImportData.Materials.Last().MaterialImportName = DefaultMaterial->GetName();
+			}
+		}
+	}
+}
+
+/**
+ * Create the bones needed to hold the transforms for the destructible chunks associated with an Apex Destructible Asset.
+ * @param ImportData - SkeletalMesh import data into which we are extracting information
+ * @param ApexDestructibleAsset - the Apex Destructible Asset
+ */
+static void CreateBones(FSkeletalMeshImportData &ImportData, const apex::DestructibleAsset& ApexDestructibleAsset)
+{
+	// Just need to create ApexDestructibleAsset.getChunkCount() bones, all with identity transform poses
+	const uint32 ChunkCount = ApexDestructibleAsset.getChunkCount();
+	if( ChunkCount == 0 )
+	{
+		UE_LOG(LogApexDestructibleAssetImport, Warning,TEXT("%s has no chunks"), ANSI_TO_TCHAR(ApexDestructibleAsset.getName()) );
+		return;
+	}
+
+	const uint32 BoneCount = ChunkCount + 1;	// Adding one more bone for the root bone, required by the skeletal mesh
+
+	// Format for bone names
+	uint32 Q = ChunkCount-1;
+	int32 MaxNumberWidth = 1;
+	while ((Q /= 10) != 0)
+	{
+		++MaxNumberWidth;
+	}
+
+	// Turn parts into bones
+	for (uint32 BoneIndex=0; BoneIndex<BoneCount; ++BoneIndex)
+	{
+		ImportData.RefBonesBinary.Add(SkeletalMeshImportData::FBone() );
+		// Set bone
+		SkeletalMeshImportData::FBone& Bone = ImportData.RefBonesBinary[BoneIndex];
+		if (BoneIndex == 0)
+		{
+			// Bone 0 is the required root bone
+			Bone.Name = TEXT("Root");
+			Bone.NumChildren = ChunkCount;
+			Bone.ParentIndex = INDEX_NONE;
+		}
+		else
+		{
+			// The rest are the parts
+			Bone.Name = FString::Printf( TEXT("Part%0*d"), MaxNumberWidth, BoneIndex-1);
+			Bone.NumChildren = 0;
+			// Creates a simple "flat" hierarchy
+			Bone.ParentIndex = 0;
+		}
+
+		// Set transform to identity
+		SkeletalMeshImportData::FJointPos& JointMatrix = Bone.BonePos;
+		JointMatrix.Transform = FTransform::Identity;
+		JointMatrix.Length = 1.0f;
+		JointMatrix.XSize = 100.0f;
+		JointMatrix.YSize = 100.0f;
+		JointMatrix.ZSize = 100.0f;
+	}
+}
+
+/**
+ * Fill an FSkeletalMeshImportData with data from an APEX Destructible Asset.
+ *
+ * @param ImportData - SkeletalMesh import data into which we are extracting information
+ * @param ApexDestructibleAsset - the Apex Destructible Asset
+ * @param bHaveAllNormals - if the function is successful, this value is true iff every submesh has a normal channel
+ * @param bHaveAllTangents - if the function is successful, this value is true iff every submesh has a tangent channel
+ *
+ * @return Boolean true iff the operation is successful
+ */
+static bool FillSkelMeshImporterFromApexDestructibleAsset(FSkeletalMeshImportData& ImportData, const apex::DestructibleAsset& ApexDestructibleAsset, bool& bHaveAllNormals, bool& bHaveAllTangents)
+{
+	// The APEX Destructible Asset contains an APEX Render Mesh Asset, get a pointer to this
+	const apex::RenderMeshAsset* ApexRenderMesh = ApexDestructibleAsset.getRenderMeshAsset();
+	if (ApexRenderMesh == NULL)
+	{
+		return false;
+	}
+
+	if (ApexDestructibleAsset.getChunkCount() != ApexRenderMesh->getPartCount())
+	{
+		UE_LOG(LogApexDestructibleAssetImport, Warning,TEXT("Chunk count does not match part count.  APEX Destructible Asset with chunk instancing not yet supported."));
+		return false;
+	}
+
+	// Apex Render Mesh uses triangle lists only, currently.  No need to triangulate.
+
+	// Assume there are no vertex colors
+	ImportData.bHasVertexColors = false;
+
+	// Different submeshes can have different UV counts.  Get the max.
+	uint32 UniqueUVCount = 0;
+
+	// Count vertices and triangles
+	uint32 VertexCount = 0;
+	uint32 TriangleCount = 0;
+
+	for (uint32 SubmeshIndex = 0; SubmeshIndex < ApexRenderMesh->getSubmeshCount(); ++SubmeshIndex)
+	{
+		const apex::RenderSubmesh& Submesh = ApexRenderMesh->getSubmesh(SubmeshIndex);
+		const apex::VertexBuffer& VB = Submesh.getVertexBuffer();
+		const apex::VertexFormat& VBFormat = VB.getFormat();
+
+		// Count UV channels in this VB
+		uint32 UVNum;
+		for (UVNum = 0; UVNum < apex::VertexFormat::MAX_UV_COUNT; ++UVNum)
+		{
+			const apex::VertexFormat::BufferID BufferID = VBFormat.getSemanticID((apex::RenderVertexSemantic::Enum)(apex::RenderVertexSemantic::TEXCOORD0 + UVNum));
+			if (VBFormat.getBufferIndexFromID(BufferID) < 0)
+			{
+				break;
+			}
+		}
+		UniqueUVCount = FMath::Max<uint32>( UniqueUVCount, UVNum );
+
+		// See if this VB has a color channel
+		const apex::VertexFormat::BufferID BufferID = VBFormat.getSemanticID(apex::RenderVertexSemantic::COLOR);
+		if (VBFormat.getBufferIndexFromID(BufferID) >= 0)
+		{
+			ImportData.bHasVertexColors = true;
+		}
+
+		// Count vertices
+		VertexCount += VB.getVertexCount();
+
+		// Count triangles
+		uint32 IndexCount = 0;
+		for (uint32 PartIndex = 0; PartIndex < ApexRenderMesh->getPartCount(); ++PartIndex)
+		{
+			IndexCount += Submesh.getIndexCount(PartIndex);
+		}
+		check(IndexCount%3 == 0);
+		TriangleCount += IndexCount/3;
+	}
+
+	// One UV set is required but only import up to MAX_TEXCOORDS number of uv layers
+	ImportData.NumTexCoords = FMath::Clamp<uint32>(UniqueUVCount, 1, MAX_TEXCOORDS);
+
+	// Expand buffers in ImportData:
+	ImportData.Points.AddUninitialized(VertexCount);
+	ImportData.Influences.AddUninitialized(VertexCount);
+
+	ImportData.Wedges.AddUninitialized(3*TriangleCount);
+	uint32 WedgeIndex = 0;
+
+	ImportData.Faces.AddUninitialized(TriangleCount);
+	uint32 TriangleIndex = 0;
+
+	uint32 VertexIndexBase = 0;
+
+	// True until proven otherwise
+	bHaveAllNormals = true;
+	bHaveAllTangents = true;
+
+	uint32_t MaterialCount = 0;
+	TMap<FString, int32> MaterialMap = TMap<FString, int32>();
+
+	// APEX render meshes are organized by submesh (render elements)
+	// Looping through submeshes first, can be done either way
+	for (uint32 SubmeshIndex = 0; SubmeshIndex < ApexRenderMesh->getSubmeshCount(); ++SubmeshIndex)
+	{
+		// Submesh data
+		const apex::RenderSubmesh& Submesh = ApexRenderMesh->getSubmesh(SubmeshIndex);
+		const apex::VertexBuffer& VB = Submesh.getVertexBuffer();
+		const apex::VertexFormat& VBFormat = VB.getFormat();
+		const physx::PxU32 SubmeshVertexCount = VB.getVertexCount();
+
+		if(SubmeshVertexCount == 0)
+		{
+			// Empty submesh, but the mesh as a whole may be valid, keep looking for more submeshes.
+			continue;
+		}
+
+		FString MaterialName(ApexRenderMesh->getMaterialName(SubmeshIndex));
+ 		uint32 MaterialIndex = 0;
+ 		if (!MaterialMap.Contains(MaterialName))
+ 		{
+ 			MaterialMap.Add(MaterialName, MaterialCount);
+ 			MaterialIndex = MaterialCount;
+ 			MaterialCount++;
+ 		}
+ 		else
+ 		{
+ 			MaterialIndex = MaterialMap[MaterialName];
+ 		}
+
+		// Get VB data semantic indices:
+
+		// Positions
+		const PxI32 PositionBufferIndex = VBFormat.getBufferIndexFromID(VBFormat.getSemanticID(apex::RenderVertexSemantic::POSITION));
+		if (!VB.getBufferData(&ImportData.Points[VertexIndexBase], apex::RenderDataFormat::FLOAT3, sizeof(FVector), PositionBufferIndex, 0, SubmeshVertexCount))
+		{
+			return false;	// Need a position buffer!
+		}
+
+#if INVERT_Y_AND_V
+		for (uint32 VertexNum = 0; VertexNum < SubmeshVertexCount; ++VertexNum)
+		{
+			ImportData.Points[VertexIndexBase + VertexNum].Y *= -1.0f;
+		}
+#endif
+
+		// Normals
+		const PxI32 NormalBufferIndex = VBFormat.getBufferIndexFromID(VBFormat.getSemanticID(apex::RenderVertexSemantic::NORMAL));
+		TArray<FVector> Normals;
+		Normals.AddUninitialized(SubmeshVertexCount);
+		const bool bHaveNormals = VB.getBufferData(Normals.GetData(), apex::RenderDataFormat::FLOAT3, sizeof(FVector), NormalBufferIndex, 0, SubmeshVertexCount);
+		if (!bHaveNormals)
+		{
+			FMemory::Memset(Normals.GetData(), 0, SubmeshVertexCount*sizeof(FVector));	// Fill with zeros
+		}
+
+		// Tangents
+		const PxI32 TangentBufferIndex = VBFormat.getBufferIndexFromID(VBFormat.getSemanticID(apex::RenderVertexSemantic::TANGENT));
+		TArray<FVector> Tangents;
+		Tangents.AddUninitialized(SubmeshVertexCount);
+		const bool bHaveTangents = VB.getBufferData(Tangents.GetData(), apex::RenderDataFormat::FLOAT3, sizeof(FVector), TangentBufferIndex, 0, SubmeshVertexCount);
+		if (!bHaveTangents)
+		{
+			FMemory::Memset(Tangents.GetData(), 0, SubmeshVertexCount*sizeof(FVector));	// Fill with zeros
+		}
+
+		// Update bHaveAllNormals and bHaveAllTangents
+		bHaveAllNormals = bHaveAllNormals && bHaveNormals;
+		bHaveAllTangents = bHaveAllTangents && bHaveTangents;
+
+		// Binormomals
+		const PxI32 BinormalBufferIndex = VBFormat.getBufferIndexFromID(VBFormat.getSemanticID(apex::RenderVertexSemantic::BINORMAL));
+		TArray<FVector> Binormals;
+		Binormals.AddUninitialized(SubmeshVertexCount);
+		bool bHaveBinormals = VB.getBufferData(Binormals.GetData(), apex::RenderDataFormat::FLOAT3, sizeof(FVector), BinormalBufferIndex, 0, SubmeshVertexCount);
+		if (!bHaveBinormals)
+		{
+			bHaveBinormals = bHaveNormals && bHaveTangents;
+			for (uint32 i = 0; i < SubmeshVertexCount; ++i)
+			{
+				Binormals[i] = Normals[i]^Tangents[i];	// Build from normals and tangents.  If one of these doesn't exist we'll get (0,0,0)'s
+			}
+		}
+
+		// Colors
+		const PxI32 ColorBufferIndex = VBFormat.getBufferIndexFromID(VBFormat.getSemanticID(apex::RenderVertexSemantic::COLOR));
+		TArray<FColor> Colors;
+		Colors.AddUninitialized(SubmeshVertexCount);
+		const bool bHaveColors = VB.getBufferData(Colors.GetData(), apex::RenderDataFormat::B8G8R8A8, sizeof(FColor), ColorBufferIndex, 0, SubmeshVertexCount);
+		if (!bHaveColors)
+		{
+			FMemory::Memset(Colors.GetData(), 0xFF, SubmeshVertexCount*sizeof(FColor));	// Fill with 0xFF
+		}
+
+		// UVs
+		TArray<FVector2f> UVs[apex::VertexFormat::MAX_UV_COUNT];
+		for (uint32 UVNum = 0; UVNum < ImportData.NumTexCoords; ++UVNum)
+		{
+			const PxI32 UVBufferIndex = VBFormat.getBufferIndexFromID(VBFormat.getSemanticID((apex::RenderVertexSemantic::Enum)(apex::RenderVertexSemantic::TEXCOORD0 + UVNum)));
+			UVs[UVNum].AddUninitialized(SubmeshVertexCount);
+			if (!VB.getBufferData(&UVs[UVNum][0].X, apex::RenderDataFormat::FLOAT2, sizeof(FVector2f), UVBufferIndex, 0, SubmeshVertexCount))
+			{
+				FMemory::Memset(&UVs[UVNum][0].X, 0, SubmeshVertexCount*sizeof(FVector2f));	// Fill with zeros
+			}
+		}
+
+		// Bone indices will not be imported - they're implicitly the PartIndex
+
+		// Each submesh is partitioned into parts.  Currently we're assuming a 1-1 correspondence between chunks and parts,
+		// which means that instanced chunks are not supported.  However, we will not assume that the chunk and part ordering is the same.
+		// Therefore, instead of looping through parts, we loop through chunks here, and get the part index.
+		for (uint32 ChunkIndex = 0; ChunkIndex < ApexDestructibleAsset.getChunkCount(); ++ChunkIndex)
+		{
+			const physx::PxU32 PartIndex = ApexDestructibleAsset.getPartIndex(ChunkIndex);
+			const physx::PxU32* PartIndexBuffer = Submesh.getIndexBuffer(PartIndex);
+			const physx::PxU32* PartIndexBufferStop = PartIndexBuffer + Submesh.getIndexCount(PartIndex);
+			while (PartIndexBuffer < PartIndexBufferStop)
+			{
+				physx::PxU32 SubmeshVertexIndex[3];
+#if !INVERT_Y_AND_V
+				SubmeshVertexIndex[2] = *PartIndexBuffer++;
+				SubmeshVertexIndex[1] = *PartIndexBuffer++;
+				SubmeshVertexIndex[0] = *PartIndexBuffer++;
+#else
+				SubmeshVertexIndex[0] = *PartIndexBuffer++;
+				SubmeshVertexIndex[1] = *PartIndexBuffer++;
+				SubmeshVertexIndex[2] = *PartIndexBuffer++;
+#endif
+
+				// Fill triangle
+				SkeletalMeshImportData::FTriangle& Triangle = ImportData.Faces[TriangleIndex++];
+
+				// set the face smoothing by default. It could be any number, but not zero
+				Triangle.SmoothingGroups = 255; 
+
+				// Material index
+				Triangle.MatIndex = MaterialIndex;
+				Triangle.AuxMatIndex = 0;
+
+				// Per-vertex
+				for (uint32 V = 0; V < 3; ++V)
+				{
+					// Tangent basis
+					Triangle.TangentX[V] = Tangents[SubmeshVertexIndex[V]];
+					Triangle.TangentY[V] = Binormals[SubmeshVertexIndex[V]];
+					Triangle.TangentZ[V] = Normals[SubmeshVertexIndex[V]];
+#if INVERT_Y_AND_V
+					Triangle.TangentX[V].Y *= -1.0f;
+					Triangle.TangentY[V].Y *= -1.0f;
+					Triangle.TangentZ[V].Y *= -1.0f;
+#endif
+
+					// Wedges
+					Triangle.WedgeIndex[V] = WedgeIndex;
+					SkeletalMeshImportData::FVertex& Wedge = ImportData.Wedges[WedgeIndex++];
+					Wedge.VertexIndex = VertexIndexBase + SubmeshVertexIndex[V];
+					Wedge.MatIndex = Triangle.MatIndex;
+					Wedge.Color = Colors[SubmeshVertexIndex[V]];
+					Wedge.Reserved = 0;
+					for (uint32 UVNum = 0; UVNum < ImportData.NumTexCoords; ++UVNum)
+					{
+						const FVector2f& UV = UVs[UVNum][SubmeshVertexIndex[V]];
+#if !INVERT_Y_AND_V
+						Wedge.UVs[UVNum] = UV;
+#else
+						Wedge.UVs[UVNum] = FVector2f(UV.X, 1.0f-UV.Y);
+#endif
+					}
+				}
+			}
+
+			// Bone influences
+			const physx::PxU32 PartVertexStart = Submesh.getFirstVertexIndex(PartIndex);
+			const physx::PxU32 PartVertexStop = PartVertexStart + Submesh.getVertexCount(PartIndex);
+			for (uint32 PartVertexIndex = PartVertexStart; PartVertexIndex < PartVertexStop; ++PartVertexIndex)
+			{
+				const physx::PxU32 VertexIndex = VertexIndexBase + PartVertexIndex;
+				// Note, by using ChunkIndex instead of PartInedx we are effectively setting PartIndex = ChunkIndex, which is OK since we won't be supporting instancing with the SkeletalMesh.
+				ImportData.Influences[VertexIndex].BoneIndex = ChunkIndex + 1;	// Adding 1, since the 0 bone will have no geometry from the Apex Destructible Asset.
+				ImportData.Influences[VertexIndex].Weight = 1.0;
+				ImportData.Influences[VertexIndex].VertexIndex = VertexIndex;
+			}
+		}
+
+		VertexIndexBase += SubmeshVertexCount;
+	}
+
+	if(ImportData.Points.Num() == 0)
+	{
+		// No points were loaded, can't make a mesh.
+		UE_LOG(LogApexDestructibleAssetImport, Error, TEXT("Failed to import destructible mesh - No positions were found in any submeshes."));
+		return false;
+	}
+
+	// Create mapping from import to raw- @TODO trivial at the moment, do we need this info for destructibles?
+	ImportData.PointToRawMap.AddUninitialized(ImportData.Points.Num());
+	for(int32 PointIdx=0; PointIdx<ImportData.PointToRawMap.Num(); PointIdx++)
+	{
+		ImportData.PointToRawMap[PointIdx] = PointIdx;
+	}
+
+	return true;
+}
+
+static apex::DestructibleAsset* CreateApexDestructibleAssetFromPxStream(physx::PxFileBuf& Stream)
+{
+	// Peek into the buffer to see what kind of data it is (binary or xml)
+	NvParameterized::Serializer::SerializeType SerializeType = GApexSDK->getSerializeType(Stream);
+	// Create an NvParameterized serializer for the correct data type
+	NvParameterized::Serializer* Serializer = GApexSDK->createSerializer(SerializeType);
+
+	if (Serializer)
+	{
+		// Deserialize into a DeserializedData buffer
+		NvParameterized::Serializer::DeserializedData DeserializedData;
+		Serializer->deserialize(Stream, DeserializedData);
+		apex::Asset* ApexAsset = NULL;
+		if( DeserializedData.size() > 0 )
+		{
+			// The DeserializedData has something in it, so create an APEX asset from it
+			ApexAsset = GApexSDK->createAsset( DeserializedData[0], NULL );
+			// Make sure it's a Destructible asset
+			if (ApexAsset->getObjTypeID() != GApexModuleDestructible->getModuleID())
+			{
+				GPhysCommandHandler->DeferredRelease(ApexAsset);
+				ApexAsset = NULL;
+			}
+		}
+		// Release the serializer
+		Serializer->release();
+
+		return (apex::DestructibleAsset*)ApexAsset;
+	}
+
+	return NULL;
+}
+
+apex::DestructibleAsset* CreateApexDestructibleAssetFromBuffer(const uint8* Buffer, int32 BufferSize)
+{
+	apex::DestructibleAsset* ApexDestructibleAsset = NULL;
+
+	// Wrap Buffer with the APEX read stream class
+	physx::PxFileBuf* Stream = GApexSDK->createMemoryReadStream(Buffer, BufferSize);
+
+	if (Stream != NULL)
+	{
+		ApexDestructibleAsset = CreateApexDestructibleAssetFromPxStream(*Stream);
+
+		// Release our stream
+		GApexSDK->releaseMemoryReadStream(*Stream);
+	}
+
+	return ApexDestructibleAsset;
+}
+
+apex::DestructibleAsset* CreateApexDestructibleAssetFromFile(const FString& Filename)
+{
+	apex::DestructibleAsset* ApexDestructibleAsset = NULL;
+
+	// Create a stream to read the file
+	physx::PxFileBuf* Stream = GApexSDK->createStream(TCHAR_TO_ANSI(*Filename), PxFileBuf::OPEN_READ_ONLY);
+
+	if (Stream != NULL)
+	{
+		ApexDestructibleAsset = CreateApexDestructibleAssetFromPxStream(*Stream);
+
+		// Release our stream
+		Stream->release();
+	}
+
+	return ApexDestructibleAsset;
+}
+
+bool SetApexDestructibleAsset(UDestructibleMesh& DestructibleMesh, apex::DestructibleAsset& ApexDestructibleAsset, FSkeletalMeshImportData* OutData, EDestructibleImportOptions::Type Options)
+{
+	using namespace SkeletalMeshImportUtils;
+	DestructibleMesh.PreEditChange(NULL);
+
+	DestructibleMesh.InvalidateDeriveDataCacheGUID();
+
+	TSharedPtr<FExistingDestMeshData> ExistDestMeshDataPtr;
+	if(Options & EDestructibleImportOptions::PreserveSettings)
+	{
+		ExistDestMeshDataPtr = SaveExistingDestMeshData(&DestructibleMesh);
+	}
+	
+	// The asset is going away, which will destroy any actors created from it.  We must destroy the physics state of any destructible mesh components before we release the asset.
+	for(TObjectIterator<UDestructibleComponent> It; It; ++It)
+	{
+		UDestructibleComponent* DestructibleComponent = *It;
+		if(DestructibleComponent->SkeletalMesh == &DestructibleMesh && DestructibleComponent->IsPhysicsStateCreated())
+		{
+			DestructibleComponent->DestroyPhysicsState();
+		}
+	}
+
+	// Release old NxDestructibleAsset if it exists
+	if (DestructibleMesh.ApexDestructibleAsset != NULL && DestructibleMesh.ApexDestructibleAsset != &ApexDestructibleAsset)
+	{
+		GPhysCommandHandler->DeferredRelease(DestructibleMesh.ApexDestructibleAsset);
+	}
+
+	// BRGTODO - need to remove the render data from the ApexDestructibleAsset, no longer need it
+	// Removing const cast ... we'll have to make it non-const anyway when we modify it
+	DestructibleMesh.ApexDestructibleAsset = &ApexDestructibleAsset;
+
+	if ( !(Options&EDestructibleImportOptions::PreserveSettings) )
+	{
+		// Resize the depth parameters array to the appropriate size
+		DestructibleMesh.DefaultDestructibleParameters.DepthParameters.Init(FDestructibleDepthParameters(), ApexDestructibleAsset.getDepthCount());
+
+		// Resize the fracture effects array to the appropriate size
+		DestructibleMesh.FractureEffects.AddZeroed(ApexDestructibleAsset.getDepthCount());
+
+		// Load the UnrealEd-editable parameters from the destructible asset
+		DestructibleMesh.LoadDefaultDestructibleParametersFromApexAsset();
+	}
+		
+	// Create body setup for the destructible mesh
+	DestructibleMesh.CreateBodySetup();
+
+#if 0	// BRGTODO
+	// warning for missing smoothing group info
+	CheckSmoothingInfo(FbxMesh);
+#endif
+	
+	FSkeletalMeshImportData TempData;
+	// Fill with data from buffer
+	FSkeletalMeshImportData* SkelMeshImportDataPtr = &TempData;
+	if( OutData )
+	{
+		SkelMeshImportDataPtr = OutData;
+	}
+	
+	// Get all material names here
+	ImportMaterialsForSkelMesh(*SkelMeshImportDataPtr, ApexDestructibleAsset);
+
+	// Import animation hierarchy, although this is trivial for an Apex Destructible Asset
+	CreateBones(*SkelMeshImportDataPtr, ApexDestructibleAsset);
+
+	// Import graphics data
+	bool bHaveNormals, bHaveTangents;
+	if (!FillSkelMeshImporterFromApexDestructibleAsset(*SkelMeshImportDataPtr, ApexDestructibleAsset, bHaveNormals, bHaveTangents))
+	{
+		if (ExistDestMeshDataPtr)
+		{
+			RestoreExistingDestMeshData(ExistDestMeshDataPtr, &DestructibleMesh);
+		}
+		return false;
+	}
+
+#if 0	// BRGTODO - what is this?
+	if( SkelMeshImportDataPtr->Materials.Num() == FbxMatList.Num() )
+	{
+		// reorder material according to "SKinXX" in material name
+		SetMaterialSkinXXOrder(*SkelMeshImportDataPtr, FbxMatList );
+	}
+#endif
+
+#if 0	// BRGTODO - what is this?
+	if( ImportOptions->bSplitNonMatchingTriangles )
+	{
+		DoUnSmoothVerts(*SkelMeshImportDataPtr);
+	}
+#endif
+
+	// process materials from import data
+	ProcessImportMeshMaterials( DestructibleMesh.Materials,*SkelMeshImportDataPtr );
+	
+	// process reference skeleton from import data
+	int32 SkeletalDepth=0;
+	if(!ProcessImportMeshSkeleton(DestructibleMesh.Skeleton, DestructibleMesh.RefSkeleton, SkeletalDepth, *SkelMeshImportDataPtr))
+	{
+		if (ExistDestMeshDataPtr)
+		{
+			RestoreExistingDestMeshData(ExistDestMeshDataPtr, &DestructibleMesh);
+		}
+		return false;
+	}
+	UE_LOG(LogApexDestructibleAssetImport, Warning, TEXT("Bones digested - %i  Depth of hierarchy - %i"), DestructibleMesh.RefSkeleton.GetNum(), SkeletalDepth);
+
+	// process bone influences from import data
+	ProcessImportMeshInfluences(*SkelMeshImportDataPtr, FString("ApexMesh"));
+
+	FSkeletalMeshModel& DestructibleMeshResource = *DestructibleMesh.GetImportedModel();
+	DestructibleMeshResource.LODModels.Empty();
+	DestructibleMeshResource.EmptyOriginalReductionSourceMeshData();
+	DestructibleMeshResource.LODModels.Add(new FSkeletalMeshLODModel());
+
+	DestructibleMesh.ResetLODInfo();
+	DestructibleMesh.AddLODInfo();
+	DestructibleMesh.GetLODInfo(0)->LODHysteresis = 0.02f;
+
+	// Create initial bounding box based on expanded version of reference pose for meshes without physics assets. Can be overridden by artist.
+	FBox BoundingBox(SkelMeshImportDataPtr->Points.GetData(), SkelMeshImportDataPtr->Points.Num());
+	DestructibleMesh.SetImportedBounds(FBoxSphereBounds(BoundingBox));
+
+	// Store whether or not this mesh has vertex colors
+	DestructibleMesh.bHasVertexColors = SkelMeshImportDataPtr->bHasVertexColors;
+
+	FSkeletalMeshLODModel& LODModel = DestructibleMeshResource.LODModels[0];
+	
+	LODModel.ActiveBoneIndices.Add(0);
+
+	// Pass the number of texture coordinate sets to the LODModel.  Ensure there is at least one UV coord
+	LODModel.NumTexCoords = FMath::Max<uint32>(1,SkelMeshImportDataPtr->NumTexCoords);
+//	if( bCreateRenderData )	// We always create render data
+	{
+		// copy vertex data needed to generate skinning streams for LOD
+		TArray<FVector> LODPoints;
+		TArray<SkeletalMeshImportData::FMeshWedge> LODWedges;
+		TArray<SkeletalMeshImportData::FMeshFace> LODFaces;
+		TArray<SkeletalMeshImportData::FVertInfluence> LODInfluences;
+		TArray<int32> LODPointToRawMap;
+		SkelMeshImportDataPtr->CopyLODImportData(LODPoints,LODWedges,LODFaces,LODInfluences,LODPointToRawMap);
+
+		IMeshUtilities& MeshUtilities = FModuleManager::Get().LoadModuleChecked<IMeshUtilities>("MeshUtilities");
+
+		IMeshUtilities::MeshBuildOptions BuildOptions;
+		BuildOptions.bComputeNormals = !bHaveNormals;
+		BuildOptions.bComputeTangents = !bHaveTangents;
+		BuildOptions.bUseMikkTSpace = true;
+		BuildOptions.bComputeWeightedNormals = true;
+		BuildOptions.TargetPlatform = GetTargetPlatformManagerRef().GetRunningTargetPlatform();
+
+		// Create actual rendering data.
+		if (!MeshUtilities.BuildSkeletalMesh(DestructibleMeshResource.LODModels[0], DestructibleMesh.GetPathName(), DestructibleMesh.RefSkeleton, LODInfluences,LODWedges,LODFaces,LODPoints,LODPointToRawMap,BuildOptions))
+		{
+			DestructibleMesh.MarkAsGarbage();
+
+			return false;
+		}
+
+		if (ExistDestMeshDataPtr)
+		{
+			RestoreExistingDestMeshData(ExistDestMeshDataPtr, &DestructibleMesh);
+		}
+
+		DestructibleMesh.CalculateInvRefMatrices();
+		DestructibleMesh.PostEditChange();
+		DestructibleMesh.MarkPackageDirty();
+
+#if 0 // BRGTODO : Check, we don't need this, do we?
+		// We have to go and fix any AnimSetMeshLinkup objects that refer to this skeletal mesh, as the reference skeleton has changed.
+		for(TObjectIterator<UAnimSet> It;It;++It)
+		{
+			UAnimSet* AnimSet = *It;
+
+			// Get DestructibleMesh path name
+			FName SkelMeshName = FName( *DestructibleMesh.GetPathName() );
+
+			// See if we have already cached this Skeletal Mesh.
+			const int32* IndexPtr = AnimSet->SkelMesh2LinkupCache.Find( SkelMeshName );
+
+			if( IndexPtr )
+			{
+				AnimSet->LinkupCache( *IndexPtr ).BuildLinkup( &DestructibleMesh, AnimSet );
+			}
+		}
+#endif
+
+		// Now iterate over all skeletal mesh components re-initialising them.
+		for(TObjectIterator<UDestructibleComponent> It; It; ++It)
+		{
+			UDestructibleComponent* DestructibleComponent = *It;
+			if(DestructibleComponent->SkeletalMesh == &DestructibleMesh)
+			{
+				FComponentReregisterContext ReregisterContext(DestructibleComponent);
+			}
+		}
+	}
+
+#if INVERT_Y_AND_V
+	// Apply transformation for Y inversion
+	const physx::PxMat44 MirrorY = physx::PxMat44(physx::PxVec4(1.0f, -1.0f, 1.0f, 1.0f));
+#if !USE_TEMPORARY_TRANSFORMATION_FUNCTION
+	ApexDestructibleAsset.applyTransformation(MirrorY, 1.0f);
+#else
+	ApplyTransformationToApexDestructibleAsset( ApexDestructibleAsset, MirrorY );
+#endif
+#endif
+
+	return true;
+}
+
+bool BuildDestructibleMeshFromFractureSettings(UDestructibleMesh& DestructibleMesh, FSkeletalMeshImportData* OutData)
+{
+	bool Success = false;
+
+#if WITH_APEX
+	apex::DestructibleAsset* NewApexDestructibleAsset = NULL;
+
+#if WITH_EDITORONLY_DATA
+	if (DestructibleMesh.FractureSettings != NULL)
+	{
+		TArray<UMaterialInterface*> OverrideMaterials;
+		OverrideMaterials.SetNumUninitialized(DestructibleMesh.Materials.Num());	//save old materials
+		for (int32 MaterialIndex = 0; MaterialIndex < DestructibleMesh.Materials.Num(); ++MaterialIndex)
+		{
+			OverrideMaterials[MaterialIndex] = DestructibleMesh.Materials[MaterialIndex].MaterialInterface;
+		}
+
+		int32 MaterialCount = DestructibleMesh.FractureSettings->Materials.Num();
+
+		DestructibleMesh.Materials.SetNum(MaterialCount);
+
+		for (int32 MaterialIndex = 0; MaterialIndex < DestructibleMesh.Materials.Num(); ++MaterialIndex)
+		{
+			if (MaterialIndex < OverrideMaterials.Num() && OverrideMaterials[MaterialIndex])//if user has overridden materials use it
+			{
+				DestructibleMesh.Materials[MaterialIndex].MaterialInterface = OverrideMaterials[MaterialIndex];
+				DestructibleMesh.Materials[MaterialIndex].ImportedMaterialSlotName = OverrideMaterials[MaterialIndex]->GetFName();
+				DestructibleMesh.Materials[MaterialIndex].MaterialSlotName = OverrideMaterials[MaterialIndex]->GetFName();
+			}
+			else
+			{
+				DestructibleMesh.Materials[MaterialIndex].MaterialInterface = DestructibleMesh.FractureSettings->Materials[MaterialIndex];
+				if (DestructibleMesh.FractureSettings->Materials[MaterialIndex])
+				{
+					DestructibleMesh.Materials[MaterialIndex].ImportedMaterialSlotName = DestructibleMesh.FractureSettings->Materials[MaterialIndex]->GetFName();
+					DestructibleMesh.Materials[MaterialIndex].MaterialSlotName = DestructibleMesh.FractureSettings->Materials[MaterialIndex]->GetFName();
+				}
+			}
+		}
+
+		apex::DestructibleAssetCookingDesc DestructibleAssetCookingDesc;
+		DestructibleMesh.FractureSettings->BuildDestructibleAssetCookingDesc(DestructibleAssetCookingDesc);
+		NewApexDestructibleAsset = DestructibleMesh.FractureSettings->CreateApexDestructibleAsset(DestructibleAssetCookingDesc);
+	}
+#endif	// WITH_EDITORONLY_DATA
+
+	if (NewApexDestructibleAsset != NULL)
+	{
+		Success = SetApexDestructibleAsset(DestructibleMesh, *NewApexDestructibleAsset, OutData, EDestructibleImportOptions::PreserveSettings);
+	}
+#endif // WITH_APEX
+
+	return Success;
+}
+
+UDestructibleMesh* ImportDestructibleMeshFromApexDestructibleAsset(UObject* InParent, apex::DestructibleAsset& ApexDestructibleAsset, FName Name, EObjectFlags Flags, FSkeletalMeshImportData* OutData, EDestructibleImportOptions::Type Options)
+{
+	// The APEX Destructible Asset contains an APEX Render Mesh Asset, get a pointer to this
+	const apex::RenderMeshAsset* ApexRenderMesh = ApexDestructibleAsset.getRenderMeshAsset();
+	if (ApexRenderMesh == NULL)
+	{
+		return NULL;
+	}
+
+	// Number of submeshes (aka "elements" in Unreal)
+	const physx::PxU32 SubmeshCount = ApexRenderMesh->getSubmeshCount();
+	if (SubmeshCount == 0)
+	{
+		return NULL;
+	}
+
+	// Make sure rendering is done - so we are not changing data being used by collision drawing.
+	FlushRenderingCommands();
+
+	UDestructibleMesh* DestructibleMesh = FindObject<UDestructibleMesh>(InParent, *Name.ToString());
+	if(DestructibleMesh)
+	{
+		// we found an existing mesh (reimport), clean up LODModels for the import process
+		if(FSkeletalMeshModel* ImportedResource = DestructibleMesh->GetImportedModel())
+		{
+			// Although we flushed above to make sure the resources weren't being used, we need
+			// to flush again as the call to Empty below will call destructors on the lod models.
+			// The renderer must release the resources before that happens.
+			FlushRenderingCommands();
+
+			ImportedResource->LODModels.Empty(1);
+		}
+	}
+	else
+	{
+		// Create the new UDestructibleMesh object if the one with the same name does not exist
+		DestructibleMesh = NewObject<UDestructibleMesh>(InParent, Name, Flags);
+	}
+	
+	if (!(Options & EDestructibleImportOptions::PreserveSettings))
+	{
+		// Store the current file path and timestamp for re-import purposes
+		// @todo AssetImportData make a data class for Apex destructible assets
+		DestructibleMesh->AssetImportData = NewObject<UAssetImportData>(DestructibleMesh);
+		DestructibleMesh->AssetImportData->Update(UFactory::GetCurrentFilename());
+	}
+
+	DestructibleMesh->PreEditChange(NULL);
+
+	// Build FractureSettings from ApexDestructibleAsset in case we want to re-fracture
+#if WITH_EDITORONLY_DATA
+	DestructibleMesh->CreateFractureSettings();
+	DestructibleMesh->FractureSettings->BuildRootMeshFromApexDestructibleAsset(ApexDestructibleAsset, Options);
+	// Fill materials
+	DestructibleMesh->FractureSettings->Materials.Reset(DestructibleMesh->Materials.Num());
+	for (int32 MaterialIndex = 0; MaterialIndex < DestructibleMesh->Materials.Num(); ++MaterialIndex)
+	{
+		DestructibleMesh->FractureSettings->Materials.Insert(DestructibleMesh->Materials[MaterialIndex].MaterialInterface, MaterialIndex);
+	}
+#endif	// WITH_EDITORONLY_DATA
+
+	if (!SetApexDestructibleAsset(*DestructibleMesh, ApexDestructibleAsset, OutData, Options))
+	{
+		// should remove this destructible mesh. if not, this object causes a crash when ticking because it doesn't have proper rendering resources
+		// @TODO : creates this destructible mesh after loading data completely
+		DestructibleMesh->PostEditChange();
+		DestructibleMesh->ConditionalBeginDestroy();
+		return NULL;
+	}
+
+	return DestructibleMesh;
+}
+
+#endif // WITH_APEX
+
+PRAGMA_ENABLE_DEPRECATION_WARNINGS
+
+#endif // WITH_EDITOR
diff --git a/Engine/Plugins/Runtime/ApexDestruction/Source/ApexDestruction/Private/ApexDestructionCustomPayload.cpp b/Engine/Plugins/Runtime/ApexDestruction/Source/ApexDestruction/Private/ApexDestructionCustomPayload.cpp
new file mode 100644
--- /dev/null
+++ b/Engine/Plugins/Runtime/ApexDestruction/Source/ApexDestruction/Private/ApexDestructionCustomPayload.cpp
@@ -0,0 +1,111 @@
+// Copyright Epic Games, Inc. All Rights Reserved.
+
+#include "ApexDestructionCustomPayload.h"
+#include "DestructibleComponent.h"
+#include "PhysXPublic.h"
+#include "ApexDestructionModule.h"
+#include "AI/NavigationSystemBase.h"
+
+#if WITH_APEX
+
+void FApexDestructionSyncActors::BuildSyncData_AssumesLocked(const TArray<physx::PxRigidActor*>& ActiveActors)
+{
+PRAGMA_DISABLE_DEPRECATION_WARNINGS
+	//We want to consolidate the transforms so that we update each destructible component once by passing it an array of chunks to update.
+	//This helps avoid a lot of duplicated work like marking render dirty, computing inverse world component, etc...
+
+	//prepare map to update destructible components
+	TArray<PxShape*> Shapes;
+	for (const PxRigidActor* RigidActor : ActiveActors)
+	{
+		if (const FApexDestructionCustomPayload* DestructibleChunkInfo = ((FApexDestructionCustomPayload*)FPhysxUserData::Get<FCustomPhysXPayload>(RigidActor->userData)))
+		{
+			if (GApexModuleDestructible->owns(RigidActor) && DestructibleChunkInfo->OwningComponent.IsValid())
+			{
+				Shapes.AddUninitialized(RigidActor->getNbShapes());
+				int32 NumShapes = RigidActor->getShapes(Shapes.GetData(), Shapes.Num());
+				for (int32 ShapeIdx = 0; ShapeIdx < Shapes.Num(); ++ShapeIdx)
+				{
+					PxShape* Shape = Shapes[ShapeIdx];
+					int32 ChunkIndex;
+					if (apex::DestructibleActor* DestructibleActor = GApexModuleDestructible->getDestructibleAndChunk(Shape, &ChunkIndex))
+					{
+						const physx::PxMat44 ChunkPoseRT = DestructibleActor->getChunkPose(ChunkIndex);
+						const physx::PxTransform Transform(ChunkPoseRT);
+						if (UDestructibleComponent* DestructibleComponent = Cast<UDestructibleComponent>(FPhysxUserData::Get<UPrimitiveComponent>(DestructibleActor->userData)))
+						{
+							if (DestructibleComponent->IsRegistered())
+							{
+								TArray<FUpdateChunksInfo>& UpdateInfos = ComponentUpdateMapping.FindOrAdd(DestructibleComponent);
+								FUpdateChunksInfo* UpdateInfo = new (UpdateInfos)FUpdateChunksInfo(ChunkIndex, P2UTransform(Transform));
+							}
+						}
+					}
+				}
+
+				Shapes.Empty(Shapes.Num());	//we want to keep largest capacity array to avoid reallocs
+			}
+		}
+	}
+PRAGMA_ENABLE_DEPRECATION_WARNINGS
+}
+
+void FApexDestructionSyncActors::FinalizeSync()
+{
+PRAGMA_DISABLE_DEPRECATION_WARNINGS
+	//update each component
+	for (auto It = ComponentUpdateMapping.CreateIterator(); It; ++It)
+	{
+		if (UDestructibleComponent* DestructibleComponent = It.Key().Get())
+		{
+			TArray<FUpdateChunksInfo>& UpdateInfos = It.Value();
+			if (DestructibleComponent->IsFracturedOrInitiallyStatic())
+			{
+				DestructibleComponent->SetChunksWorldTM(UpdateInfos);
+			}
+			else
+			{
+				//if we haven't fractured it must mean that we're simulating a destructible and so we should update our GetComponentTransform() based on the single rigid body
+				DestructibleComponent->SyncComponentToRBPhysics();
+			}
+
+			FNavigationSystem::UpdateComponentData(*DestructibleComponent);
+		}
+	}
+
+	ComponentUpdateMapping.Reset();
+PRAGMA_ENABLE_DEPRECATION_WARNINGS
+}
+
+TWeakObjectPtr<UPrimitiveComponent> FApexDestructionCustomPayload::GetOwningComponent() const
+{
+	return OwningComponent;
+}
+
+int32 FApexDestructionCustomPayload::GetItemIndex() const
+{
+	return ChunkIndex;
+}
+
+FName FApexDestructionCustomPayload::GetBoneName() const
+{
+PRAGMA_DISABLE_DEPRECATION_WARNINGS
+	if(UDestructibleComponent* RawOwningComponent = OwningComponent.Get())
+	{
+		return RawOwningComponent->GetBoneName(UDestructibleComponent::ChunkIdxToBoneIdx(ChunkIndex));
+	}
+	else
+	{
+		return NAME_None;
+	}
+PRAGMA_ENABLE_DEPRECATION_WARNINGS
+}
+
+PRAGMA_DISABLE_DEPRECATION_WARNINGS
+FBodyInstance* FApexDestructionCustomPayload::GetBodyInstance() const
+{
+	return OwningComponent.IsValid() ? &OwningComponent->BodyInstance : nullptr;
+}
+PRAGMA_ENABLE_DEPRECATION_WARNINGS
+
+#endif // WITH_APEX
\ No newline at end of file
diff --git a/Engine/Plugins/Runtime/ApexDestruction/Source/ApexDestruction/Private/ApexDestructionModule.cpp b/Engine/Plugins/Runtime/ApexDestruction/Source/ApexDestruction/Private/ApexDestructionModule.cpp
new file mode 100644
--- /dev/null
+++ b/Engine/Plugins/Runtime/ApexDestruction/Source/ApexDestruction/Private/ApexDestructionModule.cpp
@@ -0,0 +1,307 @@
+// Copyright Epic Games, Inc. All Rights Reserved.
+
+#include "ApexDestructionModule.h"
+#include "PhysicsCore.h"
+#include "Modules/ModuleManager.h"
+#include "ApexDestructionCustomPayload.h"
+#include "Engine/World.h"
+#include "DestructibleComponent.h"
+#include "Physics/PhysicsInterfaceCore.h"
+#include "PhysicsInitialization.h"
+
+#ifndef APEX_STATICALLY_LINKED
+#define APEX_STATICALLY_LINKED	0
+#endif
+
+namespace nvidia
+{
+	namespace apex
+	{
+		class ModuleDestructible;
+	}
+}
+
+#if WITH_APEX
+PRAGMA_DISABLE_DEPRECATION_WARNINGS
+FApexDestructionSyncActors* FApexDestructionCustomPayload::SingletonCustomSync = nullptr;
+PRAGMA_ENABLE_DEPRECATION_WARNINGS
+nvidia::apex::ModuleDestructible* GApexModuleDestructible = nullptr;
+void* ApexDestructibleDLL = nullptr;
+
+/**
+APEX Destructible chunk report interface
+This interface delivers summaries (which can be detailed to the single chunk level, depending on the settings)
+of chunk fracture and destroy events.
+*/
+class FApexChunkReport : public nvidia::apex::UserChunkReport
+{
+public:
+	// NxUserChunkReport interface.
+
+	virtual void	onDamageNotify(const nvidia::apex::DamageEventReportData& damageEvent) override;
+	
+
+	virtual void	onStateChangeNotify(const nvidia::apex::ChunkStateEventData& visibilityEvent) override
+	{
+PRAGMA_DISABLE_DEPRECATION_WARNINGS
+		UDestructibleComponent* DestructibleComponent = Cast<UDestructibleComponent>(FPhysxUserData::Get<UPrimitiveComponent>(visibilityEvent.destructible->userData));
+		check(DestructibleComponent);
+
+		if (!IsValid(DestructibleComponent))	//don't notify if object is being destroyed
+		{
+			return;
+		}
+
+		DestructibleComponent->OnVisibilityEvent(visibilityEvent);
+PRAGMA_ENABLE_DEPRECATION_WARNINGS
+	}
+
+	virtual bool	releaseOnNoChunksVisible(const nvidia::apex::DestructibleActor* destructible) override
+	{
+		return false;
+	}
+
+	virtual void	onDestructibleWake(nvidia::apex::DestructibleActor** destructibles, physx::PxU32 count) override
+	{
+	}
+
+	virtual void	onDestructibleSleep(nvidia::apex::DestructibleActor** destructibles, physx::PxU32 count) override
+	{
+	}
+}GApexChunkReport;
+
+
+/**
+APEX PhysX3 interface
+This interface allows us to modify the PhysX simulation filter shader data with contact pair flags
+*/
+class FApexPhysX3Interface : public nvidia::apex::PhysX3Interface
+{
+public:
+	// NxApexPhysX3Interface interface.
+
+	virtual void setContactReportFlags(physx::PxShape* PShape, physx::PxPairFlags PFlags, nvidia::apex::DestructibleActor* actor, PxU16 actorChunkIndex) override
+	{
+PRAGMA_DISABLE_DEPRECATION_WARNINGS
+		UDestructibleComponent* DestructibleComponent = Cast<UDestructibleComponent>(FPhysxUserData::Get<UPrimitiveComponent>(PShape->userData));
+		check(DestructibleComponent);
+
+		DestructibleComponent->Pair(actorChunkIndex, PShape);
+PRAGMA_ENABLE_DEPRECATION_WARNINGS
+	}
+
+	virtual physx::PxPairFlags getContactReportFlags(const physx::PxShape* PShape) const override
+	{
+		PxFilterData FilterData = PShape->getSimulationFilterData();
+		return (physx::PxPairFlags)FilterData.word3;
+	}
+} GPhysX3Interface_ApexDestructionImp;
+
+static TAutoConsoleVariable<int32> CVarAPEXMaxDestructibleDynamicChunkIslandCount(
+	TEXT("p.APEXMaxDestructibleDynamicChunkIslandCount"),
+	2000,
+	TEXT("APEX Max Destructilbe Dynamic Chunk Island Count."),
+	ECVF_Default);
+
+
+static TAutoConsoleVariable<int32> CVarAPEXMaxDestructibleDynamicChunkCount(
+	TEXT("p.APEXMaxDestructibleDynamicChunkCount"),
+	2000,
+	TEXT("APEX Max Destructible dynamic Chunk Count."),
+	ECVF_Default);
+
+static TAutoConsoleVariable<int32> CVarAPEXSortDynamicChunksByBenefit(
+	TEXT("p.bAPEXSortDynamicChunksByBenefit"),
+	1,
+	TEXT("True if APEX should sort dynamic chunks by benefit."),
+	ECVF_Default);
+
+struct FPendingApexDamageEvent
+{
+PRAGMA_DISABLE_DEPRECATION_WARNINGS
+	TWeakObjectPtr<UDestructibleComponent> DestructibleComponent;
+	apex::DamageEventReportData DamageEvent;
+	TArray<apex::ChunkData> ApexChunkData;
+
+	FPendingApexDamageEvent(UDestructibleComponent* InDestructibleComponent, const apex::DamageEventReportData& InDamageEvent)
+		: DestructibleComponent(InDestructibleComponent)
+		, DamageEvent(InDamageEvent)
+	{
+		ApexChunkData.AddUninitialized(InDamageEvent.fractureEventListSize);
+		for (uint32 ChunkIdx = 0; ChunkIdx < InDamageEvent.fractureEventListSize; ++ChunkIdx)
+		{
+			ApexChunkData[ChunkIdx] = InDamageEvent.fractureEventList[ChunkIdx];
+		}
+
+		DamageEvent.fractureEventList = ApexChunkData.GetData();
+	}
+PRAGMA_ENABLE_DEPRECATION_WARNINGS
+};
+#endif
+
+class FApexDestructionModule : public IModuleInterface
+{
+private:
+	FDelegateHandle OnPhysSceneInitHandle;
+	FDelegateHandle OnPhysDispatchNotifications;
+
+#if WITH_APEX
+	TMap<FPhysScene*, TArray<FPendingApexDamageEvent>> PendingDamageEventsMap;
+#endif
+
+public:
+
+	/** IModuleInterface implementation */
+	virtual void StartupModule() override
+	{	
+#if WITH_APEX
+		ApexDestructibleDLL = PhysDLLHelper::LoadAPEXModule("Apex_Destructible");
+
+#if APEX_STATICALLY_LINKED
+		// We need to instantiate the module if we have statically linked them
+		// Otherwise all createModule functions will fail
+		instantiateModuleDestructible();
+#endif
+
+		// Load APEX Destruction module
+		GApexModuleDestructible = static_cast<apex::ModuleDestructible*>(GApexSDK->createModule("Destructible"));
+		check(GApexModuleDestructible);
+
+		// Set Destructible module parameters
+		NvParameterized::Interface* ModuleParams = GApexModuleDestructible->getDefaultModuleDesc();
+		// ModuleParams contains the default module descriptor, which may be modified here before calling the module init function
+		GApexModuleDestructible->init(*ModuleParams);
+		// Set chunk report for fracture effect callbacks
+		GApexModuleDestructible->setChunkReport(&GApexChunkReport);
+
+
+		GApexModuleDestructible->setMaxDynamicChunkIslandCount((physx::PxU32)FMath::Max(CVarAPEXMaxDestructibleDynamicChunkIslandCount.GetValueOnGameThread(), 0));
+		GApexModuleDestructible->setMaxChunkCount((physx::PxU32)FMath::Max(CVarAPEXMaxDestructibleDynamicChunkCount.GetValueOnGameThread(), 0));
+		GApexModuleDestructible->setSortByBenefit(CVarAPEXSortDynamicChunksByBenefit.GetValueOnGameThread() != 0);
+
+		GApexModuleDestructible->scheduleChunkStateEventCallback(apex::DestructibleCallbackSchedule::FetchResults);
+
+		// APEX 1.3 to preserve 1.2 behavior
+		GApexModuleDestructible->setUseLegacyDamageRadiusSpread(true);
+		GApexModuleDestructible->setUseLegacyChunkBoundsTesting(true);
+
+PRAGMA_DISABLE_DEPRECATION_WARNINGS
+		FApexDestructionCustomPayload::SingletonCustomSync = new FApexDestructionSyncActors();
+PRAGMA_ENABLE_DEPRECATION_WARNINGS
+
+		GPhysX3Interface = &GPhysX3Interface_ApexDestructionImp;
+
+#endif
+
+		OnPhysSceneInitHandle = FPhysicsDelegates::OnPhysSceneInit.AddRaw(this, &FApexDestructionModule::OnInitPhys);
+		OnPhysDispatchNotifications = FPhysicsDelegates::OnPhysDispatchNotifications.AddRaw(this, &FApexDestructionModule::OnDispatchNotifications);
+
+		Singleton = this;
+
+		UE_LOG(LogPhysicsCore, Log, TEXT("APEX is deprecated. Destruction in future will be supported using Chaos Destruction."));
+	}
+
+	virtual void OnInitPhys(FPhysScene* PhysScene)
+	{
+#if WITH_APEX
+		//Can't set this flag after scene already created
+		/*if(PxScene* PScene = PhysScene->GetPhysXScene(SceneType))
+		{
+			SCOPED_SCENE_WRITE_LOCK(PScene);
+			PScene->setFlag(PxSceneFlag::eENABLE_KINEMATIC_PAIRS, true);
+		}*/
+
+		{
+			// Make sure we use the sync scene for apex world support of destructibles in the async scene
+			apex::Scene* ApexScene = PhysScene->GetApexScene();
+			check(ApexScene);
+			PxScene* SyncPhysXScene = PhysScene->GetPxScene();
+			check(SyncPhysXScene);
+			check(GApexModuleDestructible);
+			GApexModuleDestructible->setWorldSupportPhysXScene(*ApexScene, SyncPhysXScene);
+			GApexModuleDestructible->setDamageApplicationRaycastFlags(apex::DestructibleActorRaycastFlags::AllChunks, *ApexScene);
+		}
+		
+		PendingDamageEventsMap.FindOrAdd(PhysScene);
+#endif
+	}
+PRAGMA_DISABLE_DEPRECATION_WARNINGS
+	virtual void OnDispatchNotifications(FPhysScene* PhysScene)
+	{
+#if WITH_APEX
+		TArray<FPendingApexDamageEvent>& PendingDamageEvents = PendingDamageEventsMap.FindChecked(PhysScene);
+		for (const FPendingApexDamageEvent& PendingDamageEvent : PendingDamageEvents)
+		{
+			if (UDestructibleComponent* DestructibleComponent = PendingDamageEvent.DestructibleComponent.Get())
+			{
+				const apex::DamageEventReportData & damageEvent = PendingDamageEvent.DamageEvent;
+				check(DestructibleComponent == Cast<UDestructibleComponent>(FPhysxUserData::Get<UPrimitiveComponent>(damageEvent.destructible->userData)))	//we store this as a weak pointer above in case one of the callbacks decided to call DestroyComponent
+					DestructibleComponent->OnDamageEvent(damageEvent);
+			}
+		}
+
+		PendingDamageEvents.Empty();
+#endif
+	}
+
+	void AddPendingDamageEvent(UDestructibleComponent* DestructibleComponent, const nvidia::apex::DamageEventReportData& DamageEvent)
+	{
+#if WITH_APEX
+        if (DestructibleComponent->GetWorld())
+        {
+		    FPhysScene* PhysScene = DestructibleComponent->GetWorld()->GetPhysicsScene();
+		    TArray<FPendingApexDamageEvent>& PendingDamageEvents = PendingDamageEventsMap.FindChecked(PhysScene);
+		    PendingDamageEvents.Emplace(DestructibleComponent, DamageEvent);
+        }
+#endif
+	}
+
+	virtual void ShutdownModule() override
+	{
+#if WITH_APEX
+		delete FApexDestructionCustomPayload::SingletonCustomSync;	//TODO: this should probably make sure all destructibles are removed from the physx sim
+#endif
+
+		FPhysicsDelegates::OnPhysSceneInit.Remove(OnPhysSceneInitHandle);
+
+#if WITH_APEX
+		if(ApexDestructibleDLL)
+		{
+			PhysDLLHelper::UnloadAPEXModule(ApexDestructibleDLL);
+		}
+#endif
+	}
+PRAGMA_ENABLE_DEPRECATION_WARNINGS
+
+	static FApexDestructionModule& GetSingleton()
+	{
+		return *Singleton;
+	}
+
+	static FApexDestructionModule* Singleton;
+
+};
+
+FApexDestructionModule* FApexDestructionModule::Singleton = nullptr;
+
+#if WITH_APEX
+void FApexChunkReport::onDamageNotify(const nvidia::apex::DamageEventReportData& damageEvent)
+{
+PRAGMA_DISABLE_DEPRECATION_WARNINGS
+	UDestructibleComponent* DestructibleComponent = Cast<UDestructibleComponent>(FPhysxUserData::Get<UPrimitiveComponent>(damageEvent.destructible->userData));
+	check(DestructibleComponent);
+
+	if (!IsValid(DestructibleComponent))	//don't notify if object is being destroyed
+	{
+		return;
+	}
+
+
+	FApexDestructionModule::GetSingleton().AddPendingDamageEvent(DestructibleComponent, damageEvent);
+PRAGMA_ENABLE_DEPRECATION_WARNINGS
+}
+#endif
+
+PRAGMA_DISABLE_DEPRECATION_WARNINGS
+IMPLEMENT_MODULE(FApexDestructionModule, ApexDestruction)
+PRAGMA_ENABLE_DEPRECATION_WARNINGS
\ No newline at end of file
diff --git a/Engine/Plugins/Runtime/ApexDestruction/Source/ApexDestruction/Private/DestructibleActor.cpp b/Engine/Plugins/Runtime/ApexDestruction/Source/ApexDestruction/Private/DestructibleActor.cpp
new file mode 100644
--- /dev/null
+++ b/Engine/Plugins/Runtime/ApexDestruction/Source/ApexDestruction/Private/DestructibleActor.cpp
@@ -0,0 +1,44 @@
+// Copyright Epic Games, Inc. All Rights Reserved.
+
+/*=============================================================================
+	DestructibleActor.cpp: ADestructibleActor methods.
+=============================================================================*/
+
+
+#include "DestructibleActor.h"
+#include "DestructibleComponent.h"
+#include "Engine/SkeletalMesh.h"
+
+PRAGMA_DISABLE_DEPRECATION_WARNINGS
+ADestructibleActor::ADestructibleActor(const FObjectInitializer& ObjectInitializer)
+	: Super(ObjectInitializer)
+{
+	DestructibleComponent = CreateDefaultSubobject<UDestructibleComponent>(TEXT("DestructibleComponent0"));
+	RootComponent = DestructibleComponent;
+}
+
+#if WITH_EDITOR
+bool ADestructibleActor::GetReferencedContentObjects( TArray<UObject*>& Objects ) const
+{
+	Super::GetReferencedContentObjects(Objects);
+
+	if (DestructibleComponent && DestructibleComponent->SkeletalMesh)
+	{
+		Objects.Add(DestructibleComponent->SkeletalMesh);
+	}
+	return true;
+}
+
+void ADestructibleActor::PostLoad()
+{
+	Super::PostLoad();
+
+	if (DestructibleComponent && bAffectNavigation)
+	{
+		DestructibleComponent->SetCanEverAffectNavigation(bAffectNavigation);
+		bAffectNavigation = false;
+	}
+}
+PRAGMA_ENABLE_DEPRECATION_WARNINGS
+#endif // WITH_EDITOR
+
diff --git a/Engine/Plugins/Runtime/ApexDestruction/Source/ApexDestruction/Private/DestructibleComponent.cpp b/Engine/Plugins/Runtime/ApexDestruction/Source/ApexDestruction/Private/DestructibleComponent.cpp
new file mode 100644
--- /dev/null
+++ b/Engine/Plugins/Runtime/ApexDestruction/Source/ApexDestruction/Private/DestructibleComponent.cpp
@@ -0,0 +1,1699 @@
+// Copyright Epic Games, Inc. All Rights Reserved.
+
+/*=============================================================================
+	DestructibleComponent.cpp: UDestructibleComponent methods.
+=============================================================================*/
+
+#include "DestructibleComponent.h"
+#include "EngineStats.h"
+#include "GameFramework/DamageType.h"
+#include "AI/NavigationSystemBase.h"
+#include "Particles/ParticleSystemComponent.h"
+#include "Kismet/GameplayStatics.h"
+#include "DestructibleActor.h"
+#include "PhysXPublic.h"
+#include "PhysicsEngine/BodySetup.h"
+#include "DestructibleMesh.h"
+#include "AI/NavigationSystemHelpers.h"
+#include "PhysicalMaterials/PhysicalMaterial.h"
+#include "ObjectEditorUtils.h"
+#include "Engine/StaticMesh.h"
+#include "Physics/PhysicsFiltering.h"
+#include "ApexDestructionModule.h"
+#include "PhysicsEngine/PhysicsSettings.h"
+#include "Physics/PhysicsInterfaceCore.h"
+#include "Physics/PhysicsInterfaceUtils.h"
+#include "UObject/UObjectThreadContext.h"
+#include "Engine/DamageEvents.h"
+
+#if PHYSICS_INTERFACE_PHYSX
+#include "Physics/PhysicsGeometry.h"
+#endif
+
+PRAGMA_DISABLE_DEPRECATION_WARNINGS
+
+UDestructibleComponent::UDestructibleComponent(const FObjectInitializer& ObjectInitializer)
+	: Super(ObjectInitializer)
+#if PHYSICS_INTERFACE_PHYSX
+	, PhysxUserData(this)
+#endif
+{
+
+	PrimaryComponentTick.bCanEverTick = false;
+
+	bHasCustomNavigableGeometry = EHasCustomNavigableGeometry::Yes;
+
+	static FName CollisionProfileName(TEXT("Destructible"));
+	SetCollisionProfileName(CollisionProfileName);
+
+	bAlwaysCreatePhysicsState = true;
+	SetActiveFlag(true);
+	bMultiBodyOverlap = true;
+
+	LargeChunkThreshold = 25.f;
+
+	SetComponentSpaceTransformsDoubleBuffering(false);
+
+#if WITH_PHYSX && PHYSICS_INTERFACE_PHYSX
+	// Get contact offset params
+	FBodySetupShapeIterator::GetContactOffsetParams(ContactOffsetFactor, MinContactOffset, MaxContactOffset);
+#endif //WITH_PHYSX
+}
+
+#if WITH_EDITORONLY_DATA
+void UDestructibleComponent::Serialize(FArchive& Ar)
+{
+	Super::Serialize(Ar);
+
+	if(Ar.IsLoading())
+	{
+		// Copy our skeletal mesh value to our transient variable, so it appears in slate correctly.
+		this->DestructibleMesh_DEPRECATED = GetDestructibleMesh();
+	}
+}
+#endif // WITH_EDITORONLY_DATA
+
+#if WITH_EDITOR
+void UDestructibleComponent::PostEditChangeProperty( struct FPropertyChangedEvent& PropertyChangedEvent )
+{
+	static const FName NAME_DestructibleComponent = FName(TEXT("DestructibleComponent"));
+	static const FName NAME_DestructibleMesh = FName(TEXT("DestructibleMesh"));
+
+	Super::PostEditChangeProperty(PropertyChangedEvent);
+
+	if (PropertyChangedEvent.Property != NULL)
+	{
+		if ((FObjectEditorUtils::GetCategoryFName(PropertyChangedEvent.Property) == NAME_DestructibleComponent)
+		&&  (PropertyChangedEvent.Property->GetFName() == NAME_DestructibleMesh))
+		{
+			// If our transient mesh has changed, update our skeletal mesh.
+			SetSkeletalMesh( this->DestructibleMesh_DEPRECATED );
+		}
+	}
+}
+#endif // WITH_EDITOR
+
+FBoxSphereBounds UDestructibleComponent::CalcBounds(const FTransform& LocalToWorld) const
+{
+#if WITH_APEX
+	if(ApexDestructibleActor == NULL || ApexDestructibleActor->getBounds().isEmpty())
+	{
+		// Fallback if we don't have physics, or we have empty bounds (all chunks inactive/not visible)
+		return Super::CalcBounds(LocalToWorld);
+	}
+
+	// getLocalBounds and getChunkLocalBounds are non-updating in the APEX actor, so using those will not
+	// get us the correct bounds. Instead we take the world bounds and recalculate the local bounds
+	// from it.
+	const PxBounds3 PWorldBounds = ApexDestructibleActor->getBounds();
+	
+	return FBoxSphereBounds(FBox(P2UVector(PWorldBounds.minimum), P2UVector(PWorldBounds.maximum)))
+		.TransformBy(GetComponentToWorld().Inverse())
+		.TransformBy(LocalToWorld);
+
+#else	// #if WITH_APEX
+	return Super::CalcBounds(LocalToWorld);
+#endif	// #if WITH_APEX
+}
+
+bool IsImpactDamageEnabled(const UDestructibleMesh* TheDestructibleMesh, int32 Level)
+{
+	if(TheDestructibleMesh->DefaultDestructibleParameters.DamageParameters.ImpactDamage == 0.f)
+	{
+		return false;
+	}
+
+	bool bEnableImpactDamage = false;
+	const FDestructibleDepthParameters& DepthParams = TheDestructibleMesh->DefaultDestructibleParameters.DepthParameters[Level];
+	const EImpactDamageOverride LevelOverride = DepthParams.ImpactDamageOverride; 
+
+	switch(LevelOverride)
+	{
+		case EImpactDamageOverride::IDO_On:
+		{
+			return true;
+		}
+
+		case EImpactDamageOverride::IDO_Off:
+		{
+			return false;
+		}
+
+		default:
+		{
+			//return default if we're within the default level
+		    return TheDestructibleMesh->DefaultDestructibleParameters.DamageParameters.DefaultImpactDamageDepth >= Level ? TheDestructibleMesh->DefaultDestructibleParameters.DamageParameters.bEnableImpactDamage : false;
+		}
+	}
+}
+
+void UDestructibleComponent::OnUpdateTransform(EUpdateTransformFlags UpdateTransformFlags, ETeleportType Teleport)
+{
+	// We are handling the physics move below, so don't handle it at higher levels
+	Super::OnUpdateTransform(UpdateTransformFlags | EUpdateTransformFlags::SkipPhysicsUpdate, Teleport);
+
+	if (SkeletalMesh == NULL)
+	{
+		return;
+	}
+
+	if (!bPhysicsStateCreated || !!(UpdateTransformFlags & EUpdateTransformFlags::SkipPhysicsUpdate))
+	{
+		return;
+	}
+
+	const FTransform& CurrentLocalToWorld = GetComponentTransform();
+
+#if !(UE_BUILD_SHIPPING)
+	if(CurrentLocalToWorld.ContainsNaN())
+	{
+		logOrEnsureNanError(TEXT("UDestructibleComponent:OnUpdateTransform found NaN in CurrentLocalToWorld: %s"), *CurrentLocalToWorld.ToString());
+		return;
+	}
+#endif
+
+	// warn if it has non-uniform scale
+	const FVector& MeshScale3D = CurrentLocalToWorld.GetScale3D();
+#if !(UE_BUILD_SHIPPING || UE_BUILD_TEST)
+	if( !MeshScale3D.IsUniform() )
+	{
+		UE_LOG(LogPhysics, Log, TEXT("UDestructibleComponent::SendPhysicsTransform : Non-uniform scale factor (%s) can cause physics to mismatch for %s  SkelMesh: %s"), *MeshScale3D.ToString(), *GetFullName(), SkeletalMesh ? *SkeletalMesh->GetFullName() : TEXT("NULL"));
+	}
+#endif
+
+#if WITH_APEX
+	if (ApexDestructibleActor)
+	{
+		PxRigidDynamic* PRootActor = ApexDestructibleActor->getChunkPhysXActor(0);
+		PxMat44 GlobalPose(PxMat33(U2PQuat(CurrentLocalToWorld.GetRotation())), U2PVector(CurrentLocalToWorld.GetTranslation()));
+		if(!PRootActor || PRootActor->getScene())	//either root chunk is null meaning fractured (so there's a scene), or the root has a scene
+		{
+			ApexDestructibleActor->setGlobalPose(GlobalPose);	//this assumes ignores bTeleport and treats it like it's true 
+		}else
+		{
+			PRootActor->setGlobalPose(PxTransform(GlobalPose));	//we're not in a scene yet, so place the root actor in this new position
+		}
+	}
+#endif // #if WITH_APEX
+}
+
+void UDestructibleComponent::OnCreatePhysicsState()
+{
+	// to avoid calling PrimitiveComponent, I'm just calling ActorComponent::OnCreatePhysicsState
+	// @todo lh - fix me based on the discussion with Bryan G
+	UActorComponent::OnCreatePhysicsState();
+	bPhysicsStateCreated = true;
+
+	// What we want to do with BodySetup is simply use it to store a PhysicalMaterial, and possibly some other relevant fields.  Set up pointers from the BodyInstance to the BodySetup and this component
+	UBodySetup* BodySetup = GetBodySetup();
+	BodyInstance.OwnerComponent	= this;
+	BodyInstance.BodySetup = BodySetup;
+	BodyInstance.InstanceBodyIndex = 0;
+
+#if WITH_APEX
+	if( SkeletalMesh == NULL )
+	{
+		return;
+	}
+
+	UWorld* World = GetWorld();
+	FPhysScene* PhysScene = World->GetPhysicsScene();
+	check(PhysScene);
+
+	if( GApexModuleDestructible == NULL )
+	{
+		UE_LOG(LogPhysics, Log, TEXT("UDestructibleComponent::OnCreatePhysicsState(): APEX must be enabled to init UDestructibleComponent physics.") );
+		return;
+	}
+
+	if( ApexDestructibleActor != NULL )
+	{
+		UE_LOG(LogPhysics, Log, TEXT("UDestructibleComponent::OnCreatePhysicsState(): NxDestructibleActor already created.") );
+		return;
+	}
+
+	UDestructibleMesh* TheDestructibleMesh = GetDestructibleMesh();
+	if( TheDestructibleMesh == NULL || TheDestructibleMesh->ApexDestructibleAsset == NULL)
+	{
+		UE_LOG(LogPhysics, Log, TEXT("UDestructibleComponent::OnCreatePhysicsState(): No DestructibleMesh or missing ApexDestructibleAsset.") );
+		return;
+	}
+
+	int32 ChunkCount = TheDestructibleMesh->ApexDestructibleAsset->getChunkCount();
+	// Ensure the chunks start off invisible.  RefreshBoneTransforms should make them visible.
+	for (int32 ChunkIndex = 0; ChunkIndex < ChunkCount; ++ChunkIndex)
+	{
+		SetChunkVisible(ChunkIndex, false);
+	}
+
+#if WITH_EDITOR
+	if (GIsEditor && !World->IsGameWorld())
+	{
+		// In the editor, only set the 0 chunk to be visible.
+		if (TheDestructibleMesh->ApexDestructibleAsset->getChunkCount() > 0)
+		{
+			SetChunkVisible(0, true);
+		}
+		return;
+	}
+#endif	// WITH_EDITOR
+
+	// Only create physics in the game
+	if( !World->IsGameWorld() )
+	{
+		return;
+	}
+
+	// Set template actor/body/shape properties
+
+	// Find the PhysicalMaterial we need to apply to the physics bodies.
+	UPhysicalMaterial* PhysMat = BodyInstance.GetSimplePhysicalMaterial();
+
+	// Get the default actor descriptor NvParameterized data from the asset
+	NvParameterized::Interface* ActorParams = TheDestructibleMesh->GetDestructibleActorDesc(PhysMat);
+
+	// Create PhysX transforms from GetComponentTransform()
+	const PxTransform GlobalPose(U2PVector(GetComponentTransform().GetTranslation()), U2PQuat(GetComponentTransform().GetRotation()));
+	const PxVec3 Scale = U2PVector(GetComponentTransform().GetScale3D());
+
+	// Set the transform in the actor descriptor
+	verify( NvParameterized::setParamTransform(*ActorParams,"globalPose",GlobalPose) );
+	verify( NvParameterized::setParamVec3(*ActorParams,"scale",Scale) );
+
+	// Set the (initially) dynamic flag in the actor descriptor
+	// See if we are 'static'
+	verify( NvParameterized::setParamBool(*ActorParams,"dynamic", BodyInstance.bSimulatePhysics != false) );
+
+	// Set the sleep velocity frame decay constant (was sleepVelocitySmoothingFactor) - a new feature that should help sleeping in large piles
+	verify( NvParameterized::setParamF32(*ActorParams,"sleepVelocityFrameDecayConstant", 20.0f) );
+
+	// Set up the shape desc template
+
+	// Get collision channel and response
+	FCollisionFilterData QueryFilterData, SimFilterData;
+	uint8 MoveChannel = GetCollisionObjectType();
+	FCollisionResponseContainer CollResponse;
+	if(IsCollisionEnabled())
+	{
+		// Only enable a collision response if collision is enabled
+		CollResponse = GetCollisionResponseToChannels();
+
+		LargeChunkCollisionResponse.SetCollisionResponseContainer(CollResponse);
+		SmallChunkCollisionResponse.SetCollisionResponseContainer(CollResponse);
+		SmallChunkCollisionResponse.SetResponse(ECC_Pawn, ECR_Overlap);
+	}
+	else
+	{
+		// now since by default it will all block, if collision is disabled, we need to set to ignore
+		MoveChannel = ECC_WorldStatic;
+		CollResponse.SetAllChannels(ECR_Ignore);
+		LargeChunkCollisionResponse.SetAllChannels(ECR_Ignore);
+		SmallChunkCollisionResponse.SetAllChannels(ECR_Ignore);
+	}
+
+	const bool bEnableImpactDamage = IsImpactDamageEnabled(TheDestructibleMesh, 0);
+	const bool bEnableContactModification = TheDestructibleMesh->DefaultDestructibleParameters.DamageParameters.bCustomImpactResistance && TheDestructibleMesh->DefaultDestructibleParameters.DamageParameters.ImpactResistance > 0.f;
+
+	// Passing AssetInstanceID = 0 so we'll have self-collision
+	AActor* Owner = GetOwner();
+	CreateShapeFilterData(MoveChannel, FMaskFilter(0), Owner->GetUniqueID(), CollResponse, GetUniqueID(), 0, QueryFilterData, SimFilterData, BodyInstance.bUseCCD, BodyInstance.bNotifyRigidBodyCollision || bEnableImpactDamage, false, bEnableContactModification);
+
+	// Build filterData variations for complex and simple
+	SimFilterData.Word3 |= EPDF_SimpleCollision | EPDF_ComplexCollision | EPDF_KinematicKinematicPairs;
+	QueryFilterData.Word3 |= EPDF_SimpleCollision | EPDF_ComplexCollision | EPDF_KinematicKinematicPairs;
+
+	// Set the filterData in the shape descriptor
+	verify( NvParameterized::setParamU32(*ActorParams,"p3ShapeDescTemplate.simulationFilterData.word0", SimFilterData.Word0 ) );
+	verify( NvParameterized::setParamU32(*ActorParams,"p3ShapeDescTemplate.simulationFilterData.word1", SimFilterData.Word1 ) );
+	verify( NvParameterized::setParamU32(*ActorParams,"p3ShapeDescTemplate.simulationFilterData.word2", SimFilterData.Word2 ) );
+	verify( NvParameterized::setParamU32(*ActorParams,"p3ShapeDescTemplate.simulationFilterData.word3", SimFilterData.Word3 ) );
+	verify( NvParameterized::setParamU32(*ActorParams,"p3ShapeDescTemplate.queryFilterData.word0", QueryFilterData.Word0 ) );
+	verify( NvParameterized::setParamU32(*ActorParams,"p3ShapeDescTemplate.queryFilterData.word1", QueryFilterData.Word1 ) );
+	verify( NvParameterized::setParamU32(*ActorParams,"p3ShapeDescTemplate.queryFilterData.word2", QueryFilterData.Word2 ) );
+	verify( NvParameterized::setParamU32(*ActorParams,"p3ShapeDescTemplate.queryFilterData.word3", QueryFilterData.Word3 ) );
+
+#if !WITH_CHAOS && !WITH_IMMEDIATE_PHYSX
+	// Set the PhysX material in the shape descriptor
+	const FPhysicsMaterialHandle_PhysX& MaterialHandle = PhysMat->GetPhysicsMaterial();
+	if(PxMaterial* PMaterial = MaterialHandle.Material)
+	{
+		verify(NvParameterized::setParamU64(*ActorParams, "p3ShapeDescTemplate.material", (physx::PxU64)PMaterial));
+	}
+#endif
+
+	// Set the rest off set to 0.0f since we do not inflate convexes anymore
+	verify( NvParameterized::setParamF32(*ActorParams,"p3ShapeDescTemplate.restOffset", 0.0f) );
+
+	// Set the PhysX material in the actor descriptor
+	verify( NvParameterized::setParamBool(*ActorParams,"p3ActorDescTemplate.flags.eDISABLE_GRAVITY",false) );
+	verify( NvParameterized::setParamBool(*ActorParams,"p3ActorDescTemplate.flags.eVISUALIZATION",true) );
+
+	// Set the PxActor's and PxShape's userData fields to this component's body instance
+	verify( NvParameterized::setParamU64(*ActorParams,"p3ActorDescTemplate.userData", 0 ) );
+
+	// All shapes created by this DestructibleActor will have the userdata of the owning component.
+	// We need this, as in some cases APEX is moving shapes accross actors ( ex. FormExtended structures )
+	verify( NvParameterized::setParamU64(*ActorParams,"p3ShapeDescTemplate.userData", (PxU64)&PhysxUserData ) );
+
+	// Set up the body desc template in the actor descriptor
+	verify( NvParameterized::setParamF32(*ActorParams,"p3BodyDescTemplate.angularDamping", BodyInstance.AngularDamping ) );
+	verify( NvParameterized::setParamF32(*ActorParams,"p3BodyDescTemplate.linearDamping", BodyInstance.LinearDamping ) );
+	const PxTolerancesScale& PScale = GPhysXSDK->getTolerancesScale();
+	
+	PxF32 SleepEnergyThreshold = 0.00005f*PScale.speed*PScale.speed;	// 1/1000 Default, since the speed scale is quite high
+	SleepEnergyThreshold *= BodyInstance.GetSleepThresholdMultiplier();
+	verify( NvParameterized::setParamF32(*ActorParams,"p3BodyDescTemplate.sleepThreshold", SleepEnergyThreshold) );
+//	NvParameterized::setParamF32(*ActorParams,"bodyDescTemplate.sleepDamping", SleepDamping );
+	
+	float DensityPerCubicCM = 1.0f;
+	if(PhysMat)
+	{
+		DensityPerCubicCM = PhysMat->Density;
+	}
+	verify( NvParameterized::setParamF32(*ActorParams,"p3BodyDescTemplate.density", 0.001f * DensityPerCubicCM) );	// Convert from g/cm^3 to kg/cm^3
+
+	// Enable CCD if requested
+	verify( NvParameterized::setParamBool(*ActorParams,"p3BodyDescTemplate.flags.eENABLE_CCD", BodyInstance.bUseCCD != 0) );
+	// Ask the actor to create chunk events, for more efficient visibility updates
+	verify( NvParameterized::setParamBool(*ActorParams,"createChunkEvents", true) );
+
+	// Enable hard sleeping if requested
+	verify( NvParameterized::setParamBool(*ActorParams,"useHardSleeping", bEnableHardSleeping) );
+
+	
+
+	apex::Scene* ApexScene = PhysScene->GetApexScene();
+	PxScene* PScene = PhysScene->GetPxScene();
+
+	check(ApexScene);
+
+	ChunkInfos.Reset(ChunkCount);
+	ChunkInfos.AddDefaulted(ChunkCount);
+	PhysxChunkUserData.Reset(ChunkCount);
+	PhysxChunkUserData.AddZeroed(ChunkCount);
+
+	{
+		// Lock and flush deferred command handler here to stop any currently pending deletions from affecting new actors.
+		SCOPED_SCENE_WRITE_LOCK(PScene);
+		GPhysCommandHandler->Flush();
+
+		// Create an APEX NxDestructibleActor from the Destructible asset and actor descriptor, has to be in a locked scope as it affects simulation state
+		ApexDestructibleActor = static_cast<apex::DestructibleActor*>(TheDestructibleMesh->ApexDestructibleAsset->createApexActor(*ActorParams, *ApexScene));
+		check(ApexDestructibleActor);
+	}
+	
+	// Make a backpointer to this component
+	PhysxUserData = FPhysxUserData(this);
+	ApexDestructibleActor->userData = &PhysxUserData;
+
+	// Cache cooked collision data
+	// BRGTODO : cook in asset
+	ApexDestructibleActor->cacheModuleData();
+
+	// BRGTODO : Per-actor LOD setting
+	//	ApexDestructibleActor->forcePhysicalLod( DestructibleActor->LOD );
+
+	// Start asleep if requested
+	PxRigidDynamic* PRootActor = ApexDestructibleActor->getChunkPhysXActor(0);
+
+
+	//  Put to sleep or wake up only if the component is physics-simulated
+	if(PRootActor != NULL && BodyInstance.bSimulatePhysics)
+	{
+		SCOPED_SCENE_WRITE_LOCK(PScene);	//Question, since apex is defer adding actors do we need to lock? Locking the async scene is expensive!
+
+		PRootActor->setActorFlag(PxActorFlag::eDISABLE_GRAVITY, !BodyInstance.bEnableGravity);
+
+		// Sleep/wake up as appropriate
+		if(!BodyInstance.bStartAwake)
+		{
+			ApexDestructibleActor->setChunkPhysXActorAwakeState(0, false);
+		}
+	}
+
+	UpdateBounds();
+#endif	// #if WITH_APEX
+}
+
+void UDestructibleComponent::OnDestroyPhysicsState()
+{
+#if WITH_APEX
+	if(ApexDestructibleActor != NULL)
+	{
+		GPhysCommandHandler->DeferredRelease(ApexDestructibleActor);
+		ApexDestructibleActor = NULL;
+	
+#if !WITH_CHAOS && !WITH_IMMEDIATE_PHYSX
+		//Destructible component uses the BodyInstance in PrimitiveComponent in a very dangerous way. It assigns PxRigidDynamic to it as it needs it.
+		//Destructible PxRigidDynamic actors can be deleted from under us as PhysX sees fit.
+		//Ideally we wouldn't ever have a dangling pointer, but in practice this is hard to avoid.
+		//In theory anyone using BodyInstance on a PrimitiveComponent should be using functions like GetBodyInstance - in which case we properly fix up the dangling pointer
+		BodyInstance.ActorHandle.SyncActor = NULL;
+#endif
+	}
+#endif	// #if WITH_APEX
+	Super::OnDestroyPhysicsState();
+}
+
+UBodySetup* UDestructibleComponent::GetBodySetup()
+{
+	if (SkeletalMesh != NULL)
+	{
+		UDestructibleMesh* TheDestructibleMesh = GetDestructibleMesh();
+
+		if (TheDestructibleMesh != NULL)
+		{
+			const UDestructibleMesh* TheDestructibleMeshConst = TheDestructibleMesh;
+			return TheDestructibleMeshConst->GetBodySetup();
+		}
+	}
+
+	return NULL;
+}
+
+bool UDestructibleComponent::CanEditSimulatePhysics()
+{
+	// if destructiblemeshcomponent, we will allow it always
+	return true;
+}
+
+void UDestructibleComponent::AddImpulse( FVector Impulse, FName BoneName /*= NAME_None*/, bool bVelChange /*= false*/ )
+{
+#if WITH_APEX
+	ExecuteOnPhysicsReadWrite([&]
+	{
+		const int32 ChunkIdx = BoneIdxToChunkIdx(GetBoneIndex(BoneName));
+		ApexDestructibleActor->addForce(ChunkIdx, U2PVector(Impulse),  bVelChange ? PxForceMode::eVELOCITY_CHANGE : PxForceMode::eIMPULSE);
+	});
+#endif
+}
+
+void UDestructibleComponent::AddImpulseAtLocation( FVector Impulse, FVector Position, FName BoneName /*= NAME_None*/ )
+{
+#if WITH_APEX
+	ExecuteOnPhysicsReadWrite([&]
+	{
+		const int32 ChunkIdx = BoneIdxToChunkIdx(GetBoneIndex(BoneName));
+		PxVec3 PxPosition = U2PVector(Position);
+
+		ApexDestructibleActor->addForce(ChunkIdx, U2PVector(Impulse),  PxForceMode::eIMPULSE, &PxPosition);
+	});
+#endif
+}
+
+void UDestructibleComponent::AddForce( FVector Force, FName BoneName /*= NAME_None*/, bool bAccelChange /* = false */ )
+{
+#if WITH_APEX
+	ExecuteOnPhysicsReadWrite([&]
+	{
+		const int32 ChunkIdx = BoneIdxToChunkIdx(GetBoneIndex(BoneName));
+		ApexDestructibleActor->addForce(ChunkIdx, U2PVector(Force), bAccelChange ? PxForceMode::eACCELERATION : PxForceMode::eFORCE);
+	});
+#endif
+}
+
+void UDestructibleComponent::AddForceAtLocation( FVector Force, FVector Location, FName BoneName /*= NAME_None*/ )
+{
+#if WITH_APEX
+	ExecuteOnPhysicsReadWrite([&]
+	{
+		int32 ChunkIdx = BoneIdxToChunkIdx(GetBoneIndex(BoneName));
+		PxVec3 Position = U2PVector(Location);
+
+		ApexDestructibleActor->addForce(ChunkIdx, U2PVector(Force), PxForceMode::eFORCE, &Position);
+	});
+#endif
+}
+
+void UDestructibleComponent::AddForceAtLocationLocal(FVector Force, FVector Location, FName BoneName /*= NAME_None*/)
+{
+#if WITH_APEX
+	// AddForceAtLocation already seems to be relative for ApexDestructibles.
+	AddForceAtLocation(Force, Location, BoneName);
+#endif
+}
+
+void UDestructibleComponent::AddRadialImpulse(FVector Origin, float Radius, float Strength, ERadialImpulseFalloff Falloff, bool bVelChange)
+{
+#if WITH_APEX
+	if (bIgnoreRadialImpulse)
+	{
+		return;
+	}
+
+	if (ApexDestructibleActor == NULL)
+	{
+		return;
+	}
+
+	PxRigidDynamic** PActorBuffer = NULL;
+	PxU32 PActorCount = 0;
+	if (ApexDestructibleActor->acquirePhysXActorBuffer(PActorBuffer, PActorCount, apex::DestructiblePhysXActorQueryFlags::Dynamic))
+	{
+		PxScene* LockedScene = NULL;
+
+		while (PActorCount--)
+		{
+			PxRigidDynamic* PActor = *PActorBuffer++;
+			if (PActor != NULL)
+			{
+				if (!LockedScene)
+				{
+					LockedScene = PActor->getScene();
+					LockedScene->lockWrite();
+					LockedScene->lockRead();
+				}
+				AddRadialImpulseToPxRigidBody_AssumesLocked(*PActor, Origin, Radius, Strength, Falloff, bVelChange);
+			}
+		}
+
+		if (LockedScene)
+		{
+			LockedScene->unlockRead();
+			LockedScene->unlockWrite();
+			LockedScene = NULL;
+		}
+
+		ApexDestructibleActor->releasePhysXActorBuffer();
+	}
+#endif	// #if WITH_APEX
+}
+
+void UDestructibleComponent::AddRadialForce(FVector Origin, float Radius, float Strength, ERadialImpulseFalloff Falloff, bool bAccelChange /* = false */)
+{
+#if WITH_APEX
+	if(bIgnoreRadialForce)
+	{
+		return;
+	}
+
+	if (ApexDestructibleActor == NULL)
+	{
+		return;
+	}
+
+	PxRigidDynamic** PActorBuffer = NULL;
+	PxU32 PActorCount = 0;
+	if (ApexDestructibleActor->acquirePhysXActorBuffer(PActorBuffer, PActorCount, apex::DestructiblePhysXActorQueryFlags::Dynamic))
+	{
+		PxScene* LockedScene = NULL;
+		
+
+		while (PActorCount--)
+		{
+			PxRigidDynamic* PActor = *PActorBuffer++;
+			if (PActor != NULL)
+			{
+				if (!LockedScene)
+				{
+					LockedScene = PActor->getScene();
+					LockedScene->lockWrite();
+					LockedScene->lockRead();
+				}
+
+				AddRadialForceToPxRigidBody_AssumesLocked(*PActor, Origin, Radius, Strength, Falloff, bAccelChange);
+			}
+
+			if (LockedScene)
+			{
+				LockedScene->unlockRead();
+				LockedScene->unlockWrite();
+				LockedScene = NULL;
+			}
+		}
+		ApexDestructibleActor->releasePhysXActorBuffer();
+	}
+#endif	// #if WITH_APEX
+}
+
+void UDestructibleComponent::ReceiveComponentDamage(float DamageAmount, FDamageEvent const& DamageEvent, AController* EventInstigator, AActor* DamageCauser)
+{
+	UDamageType const* const DamageTypeCDO = DamageEvent.DamageTypeClass ? DamageEvent.DamageTypeClass->GetDefaultObject<UDamageType>() : GetDefault<UDamageType>();
+	if (DamageEvent.IsOfType(FPointDamageEvent::ClassID))
+	{
+		FPointDamageEvent const* const PointDamageEvent = (FPointDamageEvent*)(&DamageEvent);
+		ApplyDamage(DamageAmount, PointDamageEvent->HitInfo.ImpactPoint, PointDamageEvent->ShotDirection, DamageTypeCDO->DestructibleImpulse);
+	}
+	else if (DamageEvent.IsOfType(FRadialDamageEvent::ClassID))
+	{
+		FRadialDamageEvent const* const RadialDamageEvent = (FRadialDamageEvent*)(&DamageEvent);
+		ApplyRadiusDamage(DamageAmount, RadialDamageEvent->Origin, RadialDamageEvent->Params.OuterRadius, DamageTypeCDO->DestructibleImpulse, false);
+	}
+}
+
+#if WITH_APEX
+void UDestructibleComponent::SpawnFractureEffectsFromDamageEvent(const apex::DamageEventReportData& InDamageEvent)
+{
+	// Use the component's fracture effects if the override is selected, otherwise use fracture effects from the asset
+	TArray<FFractureEffect>& UseFractureEffects = (bFractureEffectOverride || !SkeletalMesh) ? FractureEffects : CastChecked<UDestructibleMesh>(SkeletalMesh)->FractureEffects;
+
+	UDestructibleMesh* TheDestructibleMesh = GetDestructibleMesh();
+	if (!TheDestructibleMesh)
+	{
+		return;
+	}
+
+
+	// We keep track of the handled parent chunks here
+	TArray<int32> HandledParents;
+	for(physx::PxU32 eventN = 0; eventN < InDamageEvent.fractureEventListSize; ++eventN)
+	{
+		const apex::ChunkData& chunkData = InDamageEvent.fractureEventList[eventN];
+		if( chunkData.depth < (physx::PxU32)UseFractureEffects.Num() )
+		{
+			// We can get the root chunk here as well, so make sure that the parent index is 0, even for the root chunk
+			int32 ParentIdx = FMath::Max(TheDestructibleMesh->ApexDestructibleAsset->getChunkParentIndex(chunkData.index), 0);
+			
+			// We can test a number of flags - we'll play an effect if the chunk was destroyed
+			// As we only get the fractured event here for chunks that come free, we spawn fracture
+			// effects only once per unique parent
+			if((chunkData.flags & apex::ApexChunkFlag::FRACTURED) && !HandledParents.Contains(ParentIdx))
+			{
+				FVector Position = P2UVector(chunkData.worldBounds.getCenter());
+				FFractureEffect& FractureEffect = UseFractureEffects[chunkData.depth];
+				if( FractureEffect.Sound != NULL )
+				{
+					// Spawn sound
+					UGameplayStatics::PlaySoundAtLocation( this, FractureEffect.Sound, Position );
+				}
+				if( FractureEffect.ParticleSystem != NULL )
+				{
+					// Spawn particle system
+					UParticleSystemComponent* ParticleSystemComponent = UGameplayStatics::SpawnEmitterAtLocation( this, FractureEffect.ParticleSystem, Position );
+
+					// Disable shadows, since destructibles tend to generate a lot of these
+					if (ParticleSystemComponent != NULL)
+					{
+						ParticleSystemComponent->CastShadow = false;
+					}
+				}
+
+				HandledParents.Add(ParentIdx);
+			}			
+		}
+	}
+}
+
+void UDestructibleComponent::OnDamageEvent(const apex::DamageEventReportData& InDamageEvent)
+{
+	FVector HitPosition = P2UVector(InDamageEvent.hitPosition);
+	FVector HitDirection = P2UVector(InDamageEvent.hitDirection);
+
+	OnComponentFracture.Broadcast(HitPosition, HitDirection);
+	if (ADestructibleActor * DestructibleActor = Cast<ADestructibleActor>(GetOwner()))
+	{
+		DestructibleActor->OnActorFracture.Broadcast(HitPosition, HitDirection);
+	}
+
+	SpawnFractureEffectsFromDamageEvent(InDamageEvent);
+
+	// After receiving damage, no longer receive decals.
+	if (bReceivesDecals)
+	{
+		bReceivesDecals = false;
+		MarkRenderStateDirty();
+	}
+}
+
+void UDestructibleComponent::OnVisibilityEvent(const apex::ChunkStateEventData & InVisibilityEvent)
+{
+	for (uint32 EventIndex = 0; EventIndex < InVisibilityEvent.stateEventListSize; ++EventIndex)
+	{
+		const apex::DestructibleChunkEvent &  Event = InVisibilityEvent.stateEventList[EventIndex];
+		// Right now the only events are visibility changes.  So as an optimization we won't check for the event type.
+		//				if (Event.event & physx::NxDestructibleChunkEvent::VisibilityChanged)
+		const bool bIsVisible = (Event.event & apex::DestructibleChunkEvent::ChunkVisible) != 0;
+		SetChunkVisible(Event.chunkIndex, bIsVisible);
+	}
+}
+#endif // WITH_APEX
+
+bool UDestructibleComponent::IsFracturedOrInitiallyStatic() const
+{
+	bool bInitiallyStatic = false;
+	bool bFractured = false;
+#if WITH_APEX
+	if (ApexDestructibleActor == nullptr)
+	{
+		return false;
+	}
+
+	bInitiallyStatic = !ApexDestructibleActor->isInitiallyDynamic();
+
+	if (bInitiallyStatic == false)
+	{
+		//If we have only one chunk and its index is 0 we are NOT fractured. Otherwise we must have fractured
+		const physx::PxU32 VisibleChunkCount = ApexDestructibleActor->getNumVisibleChunks();
+		if (VisibleChunkCount == 1)
+		{
+			const physx::PxU16* VisibleChunks = ApexDestructibleActor->getVisibleChunks();
+			bFractured = *VisibleChunks != 0;
+		}
+		else
+		{
+			bFractured = true;
+		}
+	}
+#endif
+	return bFractured || bInitiallyStatic;
+}
+
+bool UDestructibleComponent::ExecuteOnPhysicsReadOnly(TFunctionRef<void()> Func) const
+{
+#if WITH_APEX
+	if (ApexDestructibleActor)
+	{
+		FPhysScene* PhysScene = GetWorld()->GetPhysicsScene();
+		// Destructibles are always dynamic or kinematic, and therefore only go into one of the scenes
+		PxScene* PScene = PhysScene->GetPxScene();
+
+		SCOPED_SCENE_READ_LOCK(PScene);
+		Func();
+
+		return true;
+	}
+#endif
+
+	return false;
+}
+
+bool UDestructibleComponent::ExecuteOnPhysicsReadWrite(TFunctionRef<void()> Func) const
+{
+#if WITH_APEX
+	if (ApexDestructibleActor)
+	{
+		FPhysScene* PhysScene = GetWorld()->GetPhysicsScene();
+		PxScene* PScene = PhysScene->GetPxScene();
+
+		SCOPED_SCENE_WRITE_LOCK(PScene);
+		Func();
+
+		return true;
+	}
+#endif
+
+	return false;
+}
+
+void UDestructibleComponent::RefreshBoneTransforms(FActorComponentTickFunction* TickFunction)
+{
+}
+
+void UDestructibleComponent::SetDestructibleMesh(class UDestructibleMesh* NewMesh)
+{
+#if WITH_APEX
+	uint32 ChunkCount = NewMesh ? NewMesh->ApexDestructibleAsset->getChunkCount() : 0;
+	ChunkInfos.Reset(ChunkCount);
+	ChunkInfos.AddDefaulted(ChunkCount);
+	PhysxChunkUserData.Reset(ChunkCount);
+	PhysxChunkUserData.AddZeroed(ChunkCount);
+#endif // WITH_APEX
+
+	Super::SetSkeletalMesh( NewMesh );
+
+#if WITH_EDITORONLY_DATA
+	// If the SkeletalMesh has changed, update our transient value too.
+	DestructibleMesh_DEPRECATED = GetDestructibleMesh();
+#endif // WITH_EDITORONLY_DATA
+	
+	RecreatePhysicsState();
+}
+
+class UDestructibleMesh* UDestructibleComponent::GetDestructibleMesh()
+{
+	return Cast<UDestructibleMesh>(SkeletalMesh);
+}
+
+void UDestructibleComponent::SetSkeletalMesh(USkeletalMesh* InSkelMesh, bool bReinitPose)
+{
+	if(InSkelMesh != NULL && !InSkelMesh->IsA(UDestructibleMesh::StaticClass()))
+	{
+		// Issue warning and do nothing if this is not actually a UDestructibleMesh
+		UE_LOG(LogPhysics, Log, TEXT("UDestructibleComponent::SetSkeletalMesh(): Passed-in USkeletalMesh (%s) must be a UDestructibleMesh.  SkeletalMesh not set."), *InSkelMesh->GetPathName() );
+		return;
+	}
+
+	UDestructibleMesh* TheDestructibleMesh = static_cast<UDestructibleMesh*>(InSkelMesh);
+#if WITH_APEX
+	if(TheDestructibleMesh != NULL && TheDestructibleMesh->ApexDestructibleAsset == NULL)
+	{
+		UE_LOG(LogPhysics, Log, TEXT("UDestructibleComponent::SetSkeletalMesh(): Missing ApexDestructibleAsset on '%s'."), *InSkelMesh->GetPathName() );
+		return;
+	}
+
+	SetDestructibleMesh(TheDestructibleMesh);
+
+	if(TheDestructibleMesh != NULL)
+	{
+		// Resize the fracture effects array to the appropriate size
+		FractureEffects.SetNumZeroed(TheDestructibleMesh->ApexDestructibleAsset->getDepthCount());
+	}
+#else
+	SetDestructibleMesh(TheDestructibleMesh);
+#endif // WITH_APEX
+}
+
+FTransform UDestructibleComponent::GetSocketTransform(FName InSocketName, ERelativeTransformSpace TransformSpace) const
+{
+	FTransform ST = Super::GetSocketTransform(InSocketName, TransformSpace);
+
+	int32 BoneIdx = GetBoneIndex(InSocketName);
+
+	// As bones in a destructible might be scaled to 0 when hidden, we force a scale of 1 if we want the socket transform
+	if (BoneIdx > 0 && IsBoneHidden(BoneIdx))
+	{
+		ST.SetScale3D(FVector(1.0f, 1.0f, 1.0f));
+	}
+
+	return ST;
+}
+
+#if WITH_APEX
+void UDestructibleComponent::Pair( int32 ChunkIndex, PxShape* PShape)
+{
+	FApexDestructionCustomPayload* CI;
+	FPhysxUserData* UserData;
+
+	check(ChunkIndex < ChunkInfos.Num());
+
+	CI = &ChunkInfos[ChunkIndex];
+
+	CI->ChunkIndex = ChunkIndex;
+	CI->OwningComponent = this;
+
+	UserData = &PhysxChunkUserData[ChunkIndex];
+	UserData->Set<FCustomPhysXPayload>(UserData, CI);
+
+	PShape->userData = UserData;
+
+	PShape->getActor()->userData = UserData;
+
+	PShape->getActor()->setActorFlag(PxActorFlag::eDISABLE_GRAVITY, !BodyInstance.bEnableGravity);
+
+	// Set collision response to non-root chunks
+	if (GetDestructibleMesh()->ApexDestructibleAsset->getChunkParentIndex(ChunkIndex) >= 0)
+	{
+		SetCollisionResponseForShape(PShape, ChunkIndex);
+	}
+}
+#endif
+
+void UDestructibleComponent::SetChunkVisible( int32 ChunkIndex, bool bInVisible )
+{
+#if WITH_APEX
+	// Bone 0 is a dummy root bone
+	const int32 BoneIndex = ChunkIdxToBoneIdx(ChunkIndex);
+
+	if( bInVisible )
+	{
+		UnHideBone(BoneIndex);
+
+		if (NULL != ApexDestructibleActor)
+		{
+			physx::PxShape** PShapes;
+			const physx::PxU32 PShapeCount = ApexDestructibleActor->getChunkPhysXShapes(PShapes, ChunkIndex);
+			if (PShapeCount > 0)
+			{
+				const physx::PxMat44 ChunkPoseRT = ApexDestructibleActor->getChunkPose(ChunkIndex);	// Unscaled
+				const physx::PxTransform Transform(ChunkPoseRT);
+				if(IsFracturedOrInitiallyStatic())
+				{
+					SetChunkWorldRT(ChunkIndex, P2UQuat(Transform.q), P2UVector(Transform.p));
+				}
+				else
+				{
+					// Not yet fractured, sync the component to the physics body
+					UPrimitiveComponent::SyncComponentToRBPhysics();
+				}
+			}
+		}
+	}
+	else
+	{
+		HideBone(BoneIndex, PBO_None);
+	}
+
+	// Mark the transform as dirty, so the bounds are updated and sent to the render thread
+	MarkRenderTransformDirty();
+
+	// New bone positions need to be sent to render thread
+	MarkRenderDynamicDataDirty();
+#endif
+}
+
+
+void UDestructibleComponent::SetChunksWorldTM(const TArray<FUpdateChunksInfo>& UpdateInfos)
+{
+	const FQuat InvRotation = GetComponentTransform().GetRotation().Inverse();
+
+	for (const FUpdateChunksInfo& UpdateInfo : UpdateInfos)
+	{
+		// Bone 0 is a dummy root bone
+		const int32 BoneIndex = ChunkIdxToBoneIdx(UpdateInfo.ChunkIndex);
+		const FVector WorldTranslation	= UpdateInfo.WorldTM.GetLocation();
+		const FQuat WorldRotation		= UpdateInfo.WorldTM.GetRotation();
+
+		const FQuat BoneRotation = InvRotation*WorldRotation;
+		const FVector BoneTranslation = InvRotation.RotateVector(WorldTranslation - GetComponentTransform().GetTranslation()) / GetComponentTransform().GetScale3D();
+
+		GetEditableComponentSpaceTransforms()[BoneIndex] = FTransform(BoneRotation, BoneTranslation);
+	}
+
+	bNeedToFlipSpaceBaseBuffers = true;
+
+	// Mark the transform as dirty, so the bounds are updated and sent to the render thread
+	MarkRenderTransformDirty();
+
+	// New bone positions need to be sent to render thread
+	MarkRenderDynamicDataDirty();
+
+	//Update bone visibilty and flip the editable space base buffer
+	FinalizeBoneTransform();
+}
+
+void UDestructibleComponent::SetChunkWorldRT( int32 ChunkIndex, const FQuat& WorldRotation, const FVector& WorldTranslation )
+{
+	// Bone 0 is a dummy root bone
+	const int32 BoneIndex = ChunkIdxToBoneIdx(ChunkIndex);
+
+	// Mark the transform as dirty, so the bounds are updated and sent to the render thread
+	MarkRenderTransformDirty();
+
+	// New bone positions need to be sent to render thread
+	MarkRenderDynamicDataDirty();
+
+#if 0
+	// Scale is already applied to the GetComponentTransform() transform, and is carried into the bones _locally_.
+	// So there is no need to set scale in the bone local transforms
+	const FTransform WorldRT(WorldRotation, WorldTranslation, GetComponentTransform().GetScale3D());
+	SpaceBases(BoneIndex) = WorldRT*GetComponentTransform().Inverse();
+#elif 1
+	// More optimal form of the above
+	const FQuat BoneRotation = GetComponentTransform().GetRotation().Inverse()*WorldRotation;
+	const FVector BoneTranslation = GetComponentTransform().GetRotation().Inverse().RotateVector(WorldTranslation - GetComponentTransform().GetTranslation())/GetComponentTransform().GetScale3D();
+	GetEditableComponentSpaceTransforms()[BoneIndex] = FTransform(BoneRotation, BoneTranslation);
+#endif
+}
+
+void UDestructibleComponent::ApplyDamage(float DamageAmount, const FVector& HitLocation, const FVector& ImpulseDir, float ImpulseStrength)
+{
+#if WITH_APEX
+	if (ApexDestructibleActor != NULL)
+	{
+		const FVector& NormalizedImpactDir = ImpulseDir.GetSafeNormal();
+
+		// Transfer damage information to the APEX NxDestructibleActor interface
+		ApexDestructibleActor->applyDamage(DamageAmount, ImpulseStrength, U2PVector( HitLocation ), U2PVector( ImpulseDir ));
+	}
+#endif
+}
+
+void UDestructibleComponent::ApplyRadiusDamage(float BaseDamage, const FVector& HurtOrigin, float DamageRadius, float ImpulseStrength, bool bFullDamage)
+{
+#if WITH_APEX
+	if (ApexDestructibleActor != NULL)
+	{
+		// Transfer damage information to the APEX NxDestructibleActor interface
+		ApexDestructibleActor->applyRadiusDamage(BaseDamage, ImpulseStrength, U2PVector( HurtOrigin ), DamageRadius, bFullDamage ? false : true);
+	}
+#endif
+}
+
+#if WITH_APEX
+int32 ExportPxActorGeometry(const PxRigidDynamic* PActor, TArray<PxShape*>& Shapes, FNavigableGeometryExport& GeomExport)
+{
+	int32 ShapesExportedCount = 0;
+	if (PActor == nullptr)
+	{
+		return ShapesExportedCount;
+	}
+
+	const FTransform PActorGlobalPose = P2UTransform(PActor->getGlobalPose());
+
+	const PxU32 ShapesCount = PActor->getNbShapes();
+	if (ShapesCount > PxU32(Shapes.Num()))
+	{
+		Shapes.AddUninitialized(ShapesCount - Shapes.Num());
+	}
+	const PxU32 RetrievedShapesCount = PActor->getShapes(Shapes.GetData(), Shapes.Num());
+	PxShape* const* ShapePtr = Shapes.GetData();
+	for (PxU32 ShapeIndex = 0; ShapeIndex < RetrievedShapesCount; ++ShapeIndex, ++ShapePtr)
+	{
+		if (*ShapePtr != NULL)
+		{
+			const PxTransform LocalPose = (*ShapePtr)->getLocalPose();
+			FTransform LocalToWorld = P2UTransform(LocalPose);
+			LocalToWorld.Accumulate(PActorGlobalPose);
+
+			switch ((*ShapePtr)->getGeometryType())
+			{
+			case PxGeometryType::eCONVEXMESH:
+			{
+				PxConvexMeshGeometry Geometry;
+				if ((*ShapePtr)->getConvexMeshGeometry(Geometry))
+				{
+					++ShapesExportedCount;
+
+					// @todo address Geometry.scale not being used here
+					GeomExport.ExportPxConvexMesh(Geometry.convexMesh, LocalToWorld);
+				}
+			}
+				break;
+
+			case PxGeometryType::eTRIANGLEMESH:
+			{
+				// @todo address Geometry.scale not being used here
+				PxTriangleMeshGeometry Geometry;
+				if ((*ShapePtr)->getTriangleMeshGeometry(Geometry))
+				{
+					++ShapesExportedCount;
+					if ((Geometry.triangleMesh->getTriangleMeshFlags()) & PxTriangleMeshFlag::e16_BIT_INDICES)
+					{
+						GeomExport.ExportPxTriMesh16Bit(Geometry.triangleMesh, LocalToWorld);
+					}
+					else
+					{
+						GeomExport.ExportPxTriMesh32Bit(Geometry.triangleMesh, LocalToWorld);
+					}
+				}
+			}
+				break;
+
+			default:
+			{
+				UE_LOG(LogPhysics, Log, TEXT("UDestructibleComponent::DoCustomNavigableGeometryExport(): unhandled PxGeometryType, %d.")
+					, int32((*ShapePtr)->getGeometryType()));
+			}
+				break;
+			}
+		}
+	}
+
+	return ShapesExportedCount;
+}
+#endif
+
+DECLARE_DWORD_COUNTER_STAT(TEXT("Num Destructible shapes exported"),STAT_Navigation_DestructiblesShapesExported,STATGROUP_Navigation );
+
+bool UDestructibleComponent::DoCustomNavigableGeometryExport(FNavigableGeometryExport& GeomExport) const
+{
+	bool bExportFromBodySetup = true;
+
+#if WITH_APEX
+	if (ApexDestructibleActor == NULL)
+	{
+#if WITH_EDITORONLY_DATA
+		if (DestructibleMesh && DestructibleMesh->SourceStaticMesh)
+		{
+			GeomExport.ExportRigidBodySetup(*DestructibleMesh->SourceStaticMesh->GetBodySetup(), GetComponentTransform());
+			bExportFromBodySetup = false;
+		}
+#endif	// WITH_EDITORONLY_DATA
+
+		return bExportFromBodySetup;
+	}
+
+	apex::DestructibleActor* DestrActor = const_cast<apex::DestructibleActor*>(ApexDestructibleActor);
+
+	const FTransform ComponentToWorldNoScale(GetComponentTransform().GetRotation(), GetComponentTransform().GetTranslation(), FVector(1.f));
+	TArray<PxShape*> Shapes;
+	Shapes.AddUninitialized(8);
+	PxRigidDynamic** PActorBuffer = NULL;
+	PxU32 PActorCount = 0;
+	if (DestrActor->acquirePhysXActorBuffer(PActorBuffer, PActorCount
+		, apex::DestructiblePhysXActorQueryFlags::Static
+		| apex::DestructiblePhysXActorQueryFlags::Dormant
+		| apex::DestructiblePhysXActorQueryFlags::Dynamic))
+	{
+		uint32 ShapesExportedCount = 0;
+
+		while (PActorCount--)
+		{
+			const PxRigidDynamic* PActor = *PActorBuffer++;
+			int32 NumExported = ExportPxActorGeometry(PActor, Shapes, GeomExport);
+			ShapesExportedCount += NumExported;
+		}
+
+		ApexDestructibleActor->releasePhysXActorBuffer();
+		bExportFromBodySetup = (ShapesExportedCount == 0);
+
+		INC_DWORD_STAT_BY(STAT_Navigation_DestructiblesShapesExported, ShapesExportedCount);
+	}
+
+	// workaround for not fractured actor
+	if (bExportFromBodySetup)
+	{
+		const PxU32 VisibleChunkCount = ApexDestructibleActor->getNumVisibleChunks();
+		if (VisibleChunkCount == 1)
+		{
+			const PxU16* VisibleChunks = ApexDestructibleActor->getVisibleChunks();
+			if (VisibleChunks[0] == 0)
+			{
+				// only root chunk visible = not fractured
+				const PxRigidDynamic* PActor = ApexDestructibleActor->getChunkPhysXActor(0);
+				
+				const int32 NumExported = ExportPxActorGeometry(PActor, Shapes, GeomExport);
+				bExportFromBodySetup = (NumExported == 0);
+			}
+		}
+	}
+#endif // WITH_APEX
+
+	// we don't want a regular geometry export
+	return bExportFromBodySetup;
+}
+
+void UDestructibleComponent::Activate( bool bReset/*=false*/ )
+{
+	if (bReset || ShouldActivate()==true)
+	{
+		SetActiveFlag(true);
+	}
+}
+
+void UDestructibleComponent::Deactivate()
+{
+	if (ShouldActivate()==false)
+	{
+		SetActiveFlag(false);
+	}
+}
+
+void UDestructibleComponent::BeginPlay()
+{
+	Super::BeginPlay();
+#if WITH_DESTRUCTIBLE_DEPRECATION
+	// Deprecation message as of UE5.1
+	UE_LOG(LogDestructible, Warning, TEXT("DestructibleComponent is deprecated. Destruction is now supported by UGeometryCollection. Please update your assets before upgrading to the next release."));
+#endif
+}
+
+void UDestructibleComponent::SetCollisionResponseToChannel(ECollisionChannel Channel, ECollisionResponse NewResponse)
+{
+#if WITH_APEX
+	if(ApexDestructibleActor)
+	{
+		if (DestructibleAsset* Asset = GetDestructibleMesh()->GetApexDestructibleAsset())
+		{
+			PxU32 NumChunks = Asset->getChunkCount();
+
+			for (uint32 ChunkIdx = 0; ChunkIdx < NumChunks; ++ChunkIdx)
+			{
+				PxRigidDynamic* PxActor = ApexDestructibleActor->getChunkPhysXActor(ChunkIdx);
+				int32 BoneIndex = ChunkIdxToBoneIdx(ChunkIdx);
+
+				SetupFakeBodyInstance(PxActor, BoneIndex);
+
+				BodyInstance.SetResponseToChannel(Channel, NewResponse);
+			}
+
+			if (NumChunks > 0)
+			{
+				OnComponentCollisionSettingsChanged();
+			}
+		}
+	}
+#endif
+}
+
+void UDestructibleComponent::SetCollisionResponseToAllChannels(ECollisionResponse NewResponse)
+{
+#if WITH_APEX
+	if(ApexDestructibleActor)
+	{
+		if (DestructibleAsset* Asset = GetDestructibleMesh()->GetApexDestructibleAsset())
+		{
+			PxU32 NumChunks = Asset->getChunkCount();
+
+
+			for(uint32 ChunkIdx = 0; ChunkIdx < NumChunks; ++ChunkIdx)
+			{
+				PxRigidDynamic* PxActor = ApexDestructibleActor->getChunkPhysXActor(ChunkIdx);
+				int32 BoneIndex = ChunkIdxToBoneIdx(ChunkIdx);
+
+				SetupFakeBodyInstance(PxActor, BoneIndex);
+
+				BodyInstance.SetResponseToAllChannels(NewResponse);
+			}
+
+			if(NumChunks > 0)
+			{
+				OnComponentCollisionSettingsChanged();
+			}
+		}
+	}
+#endif
+}
+
+void UDestructibleComponent::SetCollisionResponseToChannels(const FCollisionResponseContainer& NewReponses)
+{
+#if WITH_APEX
+	if(ApexDestructibleActor)
+	{
+		if (DestructibleAsset* Asset = GetDestructibleMesh()->GetApexDestructibleAsset())
+		{
+			PxU32 NumChunks = Asset->getChunkCount();
+
+
+			for(uint32 ChunkIdx = 0; ChunkIdx < NumChunks; ++ChunkIdx)
+			{
+				PxRigidDynamic* PxActor = ApexDestructibleActor->getChunkPhysXActor(ChunkIdx);
+				int32 BoneIndex = ChunkIdxToBoneIdx(ChunkIdx);
+
+				SetupFakeBodyInstance(PxActor, BoneIndex);
+
+				BodyInstance.SetResponseToChannels(NewReponses);
+			}
+
+			if(NumChunks > 0)
+			{
+				OnComponentCollisionSettingsChanged();
+			}
+		}
+	}
+#endif
+}
+
+bool UDestructibleComponent::ShouldUpdateTransform(bool bLODHasChanged) const
+{
+	// We do not want to update bone transforms before physics has finished
+	return false;
+}
+
+bool UDestructibleComponent::LineTraceComponent( FHitResult& OutHit, const FVector Start, const FVector End, const FCollisionQueryParams& Params )
+{
+	bool bHaveHit = false;
+#if WITH_APEX
+	if (ApexDestructibleActor != NULL)
+	{
+		PxF32 HitTime = 0.0f;
+		PxVec3 HitNormal;
+		
+		int32 ChunkIdx = ApexDestructibleActor->rayCast(HitTime, HitNormal, U2PVector(Start), U2PVector(End-Start), apex::DestructibleActorRaycastFlags::AllChunks);
+
+		if (ChunkIdx != apex::ModuleDestructibleConst::INVALID_CHUNK_INDEX && HitTime <= 1.0f)
+		{
+			PxRigidDynamic* PActor = ApexDestructibleActor->getChunkPhysXActor(ChunkIdx);
+
+			if (PActor != NULL)
+			{
+				// Store body instance state
+				FFakeBodyInstanceState PrevState;
+				SetupFakeBodyInstance(PActor, ChunkIdx, &PrevState);
+				
+				bHaveHit = Super::LineTraceComponent(OutHit, Start, End, Params);
+
+				// Reset original body instance
+				ResetFakeBodyInstance(PrevState);
+			}
+		}
+	}
+#endif
+	return bHaveHit;
+}
+
+bool UDestructibleComponent::SweepComponent(FHitResult& OutHit, const FVector Start, const FVector End, const FQuat& ShapeWorldRotation, const FCollisionShape &CollisionShape, bool bTraceComplex/*=false*/)
+{
+	bool bHaveHit = false;
+#if WITH_APEX
+	if (ApexDestructibleActor != NULL)
+	{
+		PxF32 HitTime = 0.0f;
+		PxVec3 HitNormal;
+
+		int32 ChunkIdx = ApexDestructibleActor->obbSweep(HitTime, HitNormal, U2PVector(Start), U2PVector(CollisionShape.GetExtent()), PxMat33(physx::PxIdentity), U2PVector(End - Start), apex::DestructibleActorRaycastFlags::AllChunks);
+
+		if (ChunkIdx != apex::ModuleDestructibleConst::INVALID_CHUNK_INDEX && HitTime <= 1.0f)
+		{
+			PxRigidDynamic* PActor = ApexDestructibleActor->getChunkPhysXActor(ChunkIdx);
+
+			if (PActor != NULL)
+			{
+				// Store body instance state
+				FFakeBodyInstanceState PrevState;
+				SetupFakeBodyInstance(PActor, ChunkIdx, &PrevState);
+
+				bHaveHit = Super::SweepComponent(OutHit, Start, End, ShapeWorldRotation, CollisionShape, bTraceComplex);
+
+				// Reset original body instance
+				ResetFakeBodyInstance(PrevState);
+			}
+		}
+	}
+#endif
+	return bHaveHit;
+}
+
+#if WITH_APEX
+void UDestructibleComponent::SetupFakeBodyInstance( physx::PxRigidActor* NewRigidActor, int32 InstanceIdx, FFakeBodyInstanceState* PrevState )
+{
+	//This code is very dangerous, but at the moment I have no better solution:
+	//Destructible component assigns PxRigidDynamic to the BodyInstance as it needs it.
+	//Destructible PxRigidDynamic actors can be deleted from under us as PhysX sees fit.
+	//Ideally we wouldn't ever have a dangling pointer, but in practice this is hard to avoid.
+	//In theory anyone using BodyInstance on a PrimitiveComponent should be using functions like GetBodyInstance - in which case we properly fix up the dangling pointer
+
+#if WITH_CHAOS || WITH_IMMEDIATE_PHYSX
+    check(PrevState == nullptr);
+#else
+	if (PrevState != NULL)
+	{
+		PrevState->ActorSync = BodyInstance.ActorHandle.SyncActor;
+		PrevState->InstanceIndex = BodyInstance.InstanceBodyIndex;
+	}
+#endif
+
+	const UWorld* World = GetWorld();
+	const FPhysScene* PhysScene = World ? World->GetPhysicsScene() : nullptr;
+
+#if WITH_CHAOS || WITH_IMMEDIATE_PHYSX
+    check(false);
+#else
+	BodyInstance.ActorHandle.SyncActor = NewRigidActor;
+#endif
+	BodyInstance.InstanceBodyIndex = InstanceIdx;
+}
+
+void UDestructibleComponent::ResetFakeBodyInstance( FFakeBodyInstanceState& PrevState )
+{
+#if WITH_CHAOS || WITH_IMMEDIATE_PHYSX
+    check(false);
+#else
+	BodyInstance.ActorHandle.SyncActor = PrevState.ActorSync;
+#endif
+	BodyInstance.InstanceBodyIndex = PrevState.InstanceIndex;
+}
+
+#endif
+
+void UDestructibleComponent::WakeRigidBody(FName BoneName /* = NAME_None */)
+{
+#if WITH_APEX
+	ExecuteOnPhysicsReadWrite([&]
+	{
+		const int32 ChunkIdx = BoneIdxToChunkIdx(GetBoneIndex(BoneName));
+		ApexDestructibleActor->setChunkPhysXActorAwakeState(ChunkIdx, true);
+	});
+#endif
+}
+
+void UDestructibleComponent::SetSimulatePhysics(bool bSimulate)
+{
+#if WITH_APEX
+	ExecuteOnPhysicsReadWrite([&]
+	{
+		if(bSimulate)
+		{
+			ApexDestructibleActor->setDynamic();
+		}else
+		{
+			PxRigidDynamic** PActorBuffer = NULL;
+			PxU32 PActorCount = 0;
+			if (ApexDestructibleActor->acquirePhysXActorBuffer(PActorBuffer, PActorCount))
+			{
+				for(uint32 ActorIdx = 0; ActorIdx < PActorCount; ++ActorIdx)
+				{
+					PxRigidDynamic* PActor = PActorBuffer[ActorIdx];
+					if(FApexDestructionCustomPayload* ChunkInfo = (FApexDestructionCustomPayload*)FPhysxUserData::Get<FCustomPhysXPayload>(PActor->userData))
+					{
+						ApexDestructibleActor->setChunkPhysXActorAwakeState(ChunkInfo->ChunkIndex, false);
+					}
+				}
+				ApexDestructibleActor->releasePhysXActorBuffer();
+			}
+		}
+		
+	});
+#endif
+}
+
+void UDestructibleComponent::SetEnableGravity(bool bGravityEnabled)
+{
+	Super::SetEnableGravity(bGravityEnabled);
+	
+#if WITH_APEX
+	ExecuteOnPhysicsReadWrite([&]
+	{
+		uint32 ChunkCount = ApexDestructibleActor->getNumVisibleChunks();
+		const uint16* ChunkIndices = ApexDestructibleActor->getVisibleChunks();
+		for(uint32 c = 0; c < ChunkCount; c++)
+		{
+			PxActor* PActor = ApexDestructibleActor->getChunkPhysXActor(ChunkIndices[c]);
+			check(PActor);
+			PActor->setActorFlag(PxActorFlag::eDISABLE_GRAVITY, !bGravityEnabled);
+		}
+	});
+#endif //WITH_APEX
+}
+
+FBodyInstance* UDestructibleComponent::GetBodyInstance( FName BoneName /*= NAME_None*/, bool, int) const
+{
+#if WITH_APEX
+	if (ApexDestructibleActor != NULL)
+	{
+		int32 BoneIdx = GetBoneIndex(BoneName);
+		PxRigidDynamic* PActor = ApexDestructibleActor->getChunkPhysXActor(BoneIdxToChunkIdx(BoneIdx));
+
+		const_cast<UDestructibleComponent*>(this)->SetupFakeBodyInstance(PActor, BoneIdx);
+	}
+#endif // WITH_APEX
+
+	return (FBodyInstance*)&BodyInstance;
+}
+
+bool UDestructibleComponent::IsAnySimulatingPhysics() const 
+{
+	return !!BodyInstance.bSimulatePhysics;
+}
+
+#if PHYSICS_INTERFACE_PHYSX
+
+bool UDestructibleComponent::IsChunkLarge(PxRigidActor* ChunkActor) const
+{
+#if WITH_APEX
+	check(ChunkActor);
+	return ChunkActor->getWorldBounds().getExtents().maxElement() > LargeChunkThreshold;
+#else
+	return true;
+#endif // WITH_APEX
+}
+
+#endif	//WITH_PHYSX
+
+
+void UDestructibleComponent::OnActorEnableCollisionChanged()
+{
+	ECollisionEnabled::Type NewCollisionType = GetBodyInstance()->GetCollisionEnabled();
+	SetCollisionEnabled(NewCollisionType);
+}
+
+void UDestructibleComponent::SetCollisionEnabled(ECollisionEnabled::Type NewType)
+{
+#if WITH_APEX
+	ExecuteOnPhysicsReadWrite([&]
+	{
+		PxShape** ShapeBuffer;
+		PxU32 ShapeCount = 0;
+		
+		PxU32 NumChunks = GetDestructibleMesh()->GetApexDestructibleAsset()->getChunkCount();
+		
+		const bool bSetQuery = CollisionEnabledHasQuery(NewType);
+		const bool bSetSim = CollisionEnabledHasPhysics(NewType);
+		
+		for(uint32 ChunkIdx = 0; ChunkIdx < NumChunks; ++ChunkIdx)
+		{
+			ShapeCount = ApexDestructibleActor->getChunkPhysXShapes(ShapeBuffer, ChunkIdx);
+			
+			for(uint32 ShapeIdx = 0; ShapeIdx < ShapeCount; ++ShapeIdx)
+			{
+				if(PxShape* Shape = ShapeBuffer[ShapeIdx])
+				{
+					Shape->setFlag(PxShapeFlag::eSCENE_QUERY_SHAPE, bSetQuery);
+					Shape->setFlag(PxShapeFlag::eSIMULATION_SHAPE, bSetSim);
+				}
+			}
+		}
+	});
+
+	EnsurePhysicsStateCreated();
+	OnComponentCollisionSettingsChanged();
+#endif // WITH_APEX
+}
+
+void UDestructibleComponent::SetCollisionProfileName(FName InCollisionProfileName, bool bUpdateOverlaps)
+{
+    FBodyInstance* LocalInstance = GetBodyInstance();
+    if (!LocalInstance)
+    {
+        return;
+    }
+
+	FUObjectThreadContext& ThreadContext = FUObjectThreadContext::Get();
+	if (ThreadContext.ConstructedObject == this)
+	{
+		// If we are in our constructor, defer setup until PostInitProperties as derived classes
+		// may call SetCollisionProfileName more than once.
+		LocalInstance->SetCollisionProfileNameDeferred(InCollisionProfileName);
+	}
+	else
+	{
+		ECollisionEnabled::Type OldCollisionEnabled = LocalInstance->GetCollisionEnabled();
+		LocalInstance->SetCollisionProfileName(InCollisionProfileName);
+
+		ECollisionEnabled::Type NewCollisionEnabled = LocalInstance->GetCollisionEnabled();
+
+		if (OldCollisionEnabled != NewCollisionEnabled)
+		{
+			EnsurePhysicsStateCreated();
+		}
+		OnComponentCollisionSettingsChanged(bUpdateOverlaps);
+	}
+}
+
+
+#if PHYSICS_INTERFACE_PHYSX
+
+void UDestructibleComponent::SetCollisionResponseForActor(PxRigidDynamic* Actor, int32 ChunkIdx, const FCollisionResponseContainer* ResponseOverride /*= NULL*/)
+{
+#if WITH_APEX
+	if (ApexDestructibleActor == NULL)
+	{
+		return;
+	}
+
+	// Get collision channel and response
+	FCollisionFilterData QueryFilterData, SimFilterData;
+	uint8 MoveChannel = GetCollisionObjectType();
+	if(IsCollisionEnabled())
+	{
+		UDestructibleMesh* TheDestructibleMesh = GetDestructibleMesh();
+		AActor* Owner = GetOwner();
+		bool bLargeChunk = IsChunkLarge(Actor);
+		const FCollisionResponseContainer& UseResponse = ResponseOverride == NULL ? (bLargeChunk ? LargeChunkCollisionResponse.GetResponseContainer() : SmallChunkCollisionResponse.GetResponseContainer()) : *ResponseOverride;
+
+		physx::PxU32 SupportDepth = TheDestructibleMesh->ApexDestructibleAsset->getChunkDepth(ChunkIdx);
+
+		const bool bEnableImpactDamage = IsImpactDamageEnabled(TheDestructibleMesh, SupportDepth);
+		CreateShapeFilterData(MoveChannel, FMaskFilter(0), Owner->GetUniqueID(), UseResponse, GetUniqueID(), ChunkIdxToBoneIdx(ChunkIdx), QueryFilterData, SimFilterData, BodyInstance.bUseCCD, BodyInstance.bNotifyRigidBodyCollision || bEnableImpactDamage, false);
+		
+		QueryFilterData.Word3 |= EPDF_SimpleCollision | EPDF_ComplexCollision | EPDF_KinematicKinematicPairs;
+		SimFilterData.Word3 |= EPDF_SimpleCollision | EPDF_ComplexCollision | EPDF_KinematicKinematicPairs;
+
+		SCOPED_SCENE_WRITE_LOCK(Actor->getScene());
+
+		TArray<PxShape*> Shapes;
+		Shapes.AddUninitialized(Actor->getNbShapes());
+
+		int ShapeCount = Actor->getShapes(Shapes.GetData(), Shapes.Num());
+
+		for (int32 i=0; i < ShapeCount; ++i)
+		{
+			PxShape* Shape = Shapes[i];
+
+			Shape->setQueryFilterData(U2PFilterData(QueryFilterData));
+			Shape->setSimulationFilterData(U2PFilterData(SimFilterData));
+			Shape->setFlag(PxShapeFlag::eSCENE_QUERY_SHAPE, true); 
+			Shape->setFlag(PxShapeFlag::eSIMULATION_SHAPE, true); 
+			Shape->setFlag(PxShapeFlag::eVISUALIZATION, true);
+		}
+	}
+#endif
+}
+
+void UDestructibleComponent::SetCollisionResponseForShape(PxShape* Shape, int32 ChunkIdx)
+{
+#if WITH_APEX
+	if(ApexDestructibleActor == nullptr) //since we do deferred deletion it's possible we've already meant to delete this so ignore any simulation callbacks
+	{
+		return;
+	}
+#endif
+
+	// Get collision channel and response
+	FCollisionFilterData QueryFilterData, SimFilterData;
+	uint8 MoveChannel = GetCollisionObjectType();
+	if (IsCollisionEnabled())
+	{
+		AActor* Owner = GetOwner();
+		bool bLargeChunk = IsChunkLarge(Shape->getActor());
+		const FCollisionResponse& ColResponse = bLargeChunk ? LargeChunkCollisionResponse : SmallChunkCollisionResponse;
+		//TODO: we currently assume chunks will not have impact damage as it's very expensive. Should look into exposing this a bit more
+		CreateShapeFilterData(MoveChannel, FMaskFilter(0), (Owner ? Owner->GetUniqueID() : 0), ColResponse.GetResponseContainer(), GetUniqueID(), ChunkIdxToBoneIdx(ChunkIdx), QueryFilterData, SimFilterData, BodyInstance.bUseCCD, BodyInstance.bNotifyRigidBodyCollision, false);
+
+		QueryFilterData.Word3 |= EPDF_SimpleCollision | EPDF_ComplexCollision | EPDF_KinematicKinematicPairs;
+		SimFilterData.Word3 |= EPDF_SimpleCollision | EPDF_ComplexCollision | EPDF_KinematicKinematicPairs;
+
+		SCOPED_SCENE_WRITE_LOCK(Shape->getActor()->getScene());
+
+		// See FBodySetupShapeIterator::ComputeContactOffset(const PxConvexMeshGeometry&) for this calculation. 
+		// Reproduced here to avoid going to the underlying PxConvexMeshGeometry
+		const float MinBoundsExtent = Shape->getActor()->getWorldBounds().getExtents().minElement();
+		Shape->setContactOffset(FMath::Clamp(ContactOffsetFactor * MinBoundsExtent, MinContactOffset, MaxContactOffset));
+
+		Shape->setQueryFilterData(U2PFilterData(QueryFilterData));
+		Shape->setSimulationFilterData(U2PFilterData(SimFilterData));
+		Shape->setFlag(PxShapeFlag::eSCENE_QUERY_SHAPE, true);
+		Shape->setFlag(PxShapeFlag::eSIMULATION_SHAPE, true);
+		Shape->setFlag(PxShapeFlag::eVISUALIZATION, true);
+	}
+}
+
+#endif	//WITH_PHYSX
+
+
+void UDestructibleComponent::SetMaterial(int32 ElementIndex, UMaterialInterface* Material)
+{
+	// Mesh component handles render side materials
+	Super::SetMaterial(ElementIndex, Material);
+
+	// Update physical properties of the chunks in the mesh if the body instance is valid
+	FBodyInstance* BodyInst = GetBodyInstance();
+	if (BodyInst)
+	{
+		BodyInst->UpdatePhysicalMaterials();
+	}
+	
+	// Update physical properties for individual bone instances as well
+	if (SkeletalMesh)
+	{
+		int32 NumBones = SkeletalMesh->GetRefSkeleton().GetRawBoneNum();
+		for (int32 BoneIdx = 0; BoneIdx < NumBones; ++BoneIdx)
+		{
+			FName BoneName = SkeletalMesh->GetRefSkeleton().GetBoneName(BoneIdx);
+			FBodyInstance* Instance = GetBodyInstance(BoneName);
+			if (Instance)
+			{
+				Instance->UpdatePhysicalMaterials();
+			}
+		}
+	}
+	
+
+#if WITH_APEX
+	// Set new template parameters for the apex actor, so they take effect before fracturing too.
+	if(ApexDestructibleActor)
+	{
+		nvidia::apex::PhysX3DescTemplate* Template = ApexDestructibleActor->createPhysX3DescTemplate();
+		if(ApexDestructibleActor->getPhysX3Template(*Template))
+		{
+			UPhysicalMaterial* SimpleMaterial = GetBodyInstance()->GetSimplePhysicalMaterial();
+			
+			if(SimpleMaterial)
+			{
+#if WITH_CHAOS || WITH_IMMEDIATE_PHYSX
+                check(false);
+#else
+				const FPhysicsMaterialHandle_PhysX& MaterialHandle = SimpleMaterial->GetPhysicsMaterial();
+				PxMaterial* PhysxMat = MaterialHandle.Material;
+
+				Template->setMaterials(&PhysxMat, 1);
+				ApexDestructibleActor->setPhysX3Template(Template);
+#endif
+			}
+		}
+		Template->release();
+	}
+#endif
+}
+
+PRAGMA_ENABLE_DEPRECATION_WARNINGS
diff --git a/Engine/Plugins/Runtime/ApexDestruction/Source/ApexDestruction/Private/DestructibleFractureSettings.cpp b/Engine/Plugins/Runtime/ApexDestruction/Source/ApexDestruction/Private/DestructibleFractureSettings.cpp
new file mode 100644
--- /dev/null
+++ b/Engine/Plugins/Runtime/ApexDestruction/Source/ApexDestruction/Private/DestructibleFractureSettings.cpp
@@ -0,0 +1,580 @@
+// Copyright Epic Games, Inc. All Rights Reserved.
+
+#include "DestructibleFractureSettings.h"
+#include "Materials/Material.h"
+#include "PhysXPublic.h"
+#include "PhysicsPublic.h"
+#include "Serialization/ArchiveUObjectFromStructuredArchive.h"
+
+//////////////////////////////////////////////////////////////////////////
+
+UDestructibleFractureSettings::~UDestructibleFractureSettings()
+{
+}
+
+// Local utilities
+
+#if WITH_EDITOR	//	Fracture code is only needed in editor
+
+// Why doesn't TArray have this function?
+template<typename ElementType>
+inline void Resize(TArray<ElementType>& Array, const ElementType& Item, uint32 Size)
+{
+	const uint32 OldSize = Array.Num();
+	if (Size < OldSize)
+	{
+		Array.RemoveAt(Size, OldSize-Size);
+	}
+	else
+	if (Size > OldSize)
+	{
+		Array.Reserve(Size);
+		for (uint32 Index = OldSize; Index < Size; ++Index)
+		{
+			Array.Add(Item);
+		}
+	}
+}
+
+#if WITH_APEX
+static void BuildApexRenderMesh(apex::RenderMeshAssetAuthoring& RenderMeshAssetAuthor, apex::ExplicitHierarchicalMesh& HMesh, apex::RenderDataFormat::Enum VertexNormalFormat = apex::RenderDataFormat::FLOAT3)
+{
+	// Create a mesh building descriptor
+	apex::RenderMeshAssetAuthoring::MeshDesc MeshDesc;
+	TArray<apex::RenderMeshAssetAuthoring::SubmeshDesc> SubmeshDescs;
+	SubmeshDescs.Init(apex::RenderMeshAssetAuthoring::SubmeshDesc(), HMesh.submeshCount());
+	MeshDesc.m_numSubmeshes = (PxU32)SubmeshDescs.Num();
+	MeshDesc.m_submeshes = SubmeshDescs.GetData();
+
+	// Need to transpose the submesh/part arrays.  The outer arrays on the following are indexed by submesh:
+	TArray< TArray<apex::Vertex> > SubmeshVertices;
+	SubmeshVertices.Init(TArray<apex::Vertex>(), MeshDesc.m_numSubmeshes);
+	TArray< TArray<PxU32> > SubmeshIndices;
+	SubmeshIndices.Init(TArray<PxU32>(), MeshDesc.m_numSubmeshes);
+	TArray< TArray<PxU32> > SubmeshPartIndices;
+	SubmeshPartIndices.Init(TArray<PxU32>(), MeshDesc.m_numSubmeshes);
+	TArray< TArray<PxU16> > SubmeshFlags;
+	SubmeshFlags.Init(TArray<PxU16>(), MeshDesc.m_numSubmeshes);
+	TArray< apex::RenderMeshAssetAuthoring::VertexBuffer > SubmeshVertexBuffers;
+	SubmeshVertexBuffers.Init(apex::RenderMeshAssetAuthoring::VertexBuffer(), MeshDesc.m_numSubmeshes);
+
+	for (uint32 SubmeshNum = 0; SubmeshNum < MeshDesc.m_numSubmeshes; ++SubmeshNum)
+	{
+		TArray<apex::Vertex>& Vertices = SubmeshVertices[SubmeshNum];
+		TArray<PxU32>& Indices = SubmeshIndices[SubmeshNum];
+		TArray<PxU16>& Flags = SubmeshFlags[SubmeshNum];
+		TArray<PxU32>& PartIndices = SubmeshPartIndices[SubmeshNum];
+		apex::RenderMeshAssetAuthoring::VertexBuffer& vb = SubmeshVertexBuffers[SubmeshNum];
+
+		apex::RenderMeshAssetAuthoring::SubmeshDesc& SubmeshDesc = SubmeshDescs[SubmeshNum];
+
+		SubmeshDesc.m_materialName = HMesh.submeshData(SubmeshNum)->mMaterialName;
+		apex::ExplicitVertexFormat VertexFormat = HMesh.submeshData(SubmeshNum)->mVertexFormat;
+
+		SubmeshDesc.m_numVertexBuffers = 1;
+		SubmeshDesc.m_vertexBuffers = &vb;
+
+		PartIndices.SetNumUninitialized(HMesh.partCount());
+		for(uint32 PartIndex = 0; PartIndex < HMesh.partCount(); ++PartIndex)
+		{
+			TArray<apex::Vertex> PartVertices;
+			const uint32 PartTriangleCount = HMesh.meshTriangleCount(PartIndex);
+			apex::ExplicitRenderTriangle* PartTriangles = HMesh.meshTriangles(PartIndex);
+
+			uint32 PartVertexCount = 0;
+			for(uint32 PartTriangleIndex = 0; PartTriangleIndex < PartTriangleCount; ++PartTriangleIndex)
+			{
+				if (PartTriangles[PartTriangleIndex].submeshIndex == SubmeshNum)
+				{
+					PartVertexCount	+= 3;
+				}
+			}
+
+			PartVertices.Init(apex::Vertex(), PartVertexCount);
+			PartVertexCount = 0;
+			for(uint32 PartTriangleIndex = 0; PartTriangleIndex < PartTriangleCount; ++PartTriangleIndex)
+			{
+				apex::ExplicitRenderTriangle& Triangle = PartTriangles[PartTriangleIndex];
+				if (Triangle.submeshIndex == SubmeshNum)
+				{
+					PartVertices[PartVertexCount++]	= Triangle.vertices[0];
+					PartVertices[PartVertexCount++]	= Triangle.vertices[1];
+					PartVertices[PartVertexCount++]	= Triangle.vertices[2];
+				}
+			}
+
+			PartIndices[PartIndex] = Indices.Num();
+
+			if (PartVertices.Num() > 0)
+			{
+				TArray<PxU32> Map;
+				Map.SetNumUninitialized(PartVertices.Num());
+				const PxU32 ReducedPartVertexCount = RenderMeshAssetAuthor.createReductionMap(Map.GetData(), PartVertices.GetData(), NULL, (PxU32)PartVertices.Num(), PxVec3( 0.0001f ), 0.001f, 1.0f/256.01f);
+				const PxU32 VertexPartStart = (PxU32)Vertices.Num();
+				Resize(Vertices, apex::Vertex(), VertexPartStart + ReducedPartVertexCount);
+				Resize(Indices, (uint32)0, PartIndices[PartIndex] + PartVertices.Num());
+				Resize(Flags, (uint16)0, VertexPartStart + ReducedPartVertexCount);
+				for(uint32 OldIndex = 0; OldIndex < (uint32)PartVertices.Num(); ++OldIndex)
+				{
+					const uint32 NewIndex = Map[OldIndex]+VertexPartStart;
+					Indices[PartIndices[PartIndex]+OldIndex] = NewIndex;
+					Vertices[NewIndex] = PartVertices[OldIndex];	// This will copy over several times, but with the same (or close enough) data
+					Vertices[NewIndex].boneIndices[0] = PartIndex;
+				}
+			}
+		}
+
+		SubmeshDesc.m_numVertices = (PxU32)Vertices.Num();
+		SubmeshDesc.m_numParts = (PxU32)PartIndices.Num();
+		SubmeshDesc.m_partIndices = PartIndices.GetData();
+		SubmeshDesc.m_numIndices = (PxU32)Indices.Num();
+		SubmeshDesc.m_vertexIndices = Indices.GetData();
+
+		if (SubmeshDesc.m_numParts == 0 || SubmeshDesc.m_vertexIndices == 0)
+		{
+			continue;
+		}
+
+		if (VertexFormat.mHasStaticPositions || VertexFormat.mHasDynamicPositions)
+		{
+			vb.setSemanticData( apex::RenderVertexSemantic::POSITION, &Vertices.GetData()->position, sizeof( apex::Vertex ), apex::RenderDataFormat::FLOAT3);
+		}
+		if (VertexFormat.mHasStaticNormals || VertexFormat.mHasDynamicNormals)
+		{
+			vb.setSemanticData( apex::RenderVertexSemantic::NORMAL, &Vertices.GetData()->normal, sizeof( apex::Vertex ), VertexNormalFormat, apex::RenderDataFormat::FLOAT3);
+		}
+		if (VertexFormat.mHasStaticTangents || VertexFormat.mHasDynamicTangents)
+		{
+			vb.setSemanticData( apex::RenderVertexSemantic::TANGENT, &Vertices.GetData()->tangent, sizeof( apex::Vertex ), VertexNormalFormat, apex::RenderDataFormat::FLOAT3);
+		}
+		if (VertexFormat.mHasStaticBinormals || VertexFormat.mHasDynamicBinormals)
+		{
+			vb.setSemanticData( apex::RenderVertexSemantic::BINORMAL, &Vertices.GetData()->binormal, sizeof( apex::Vertex ), VertexNormalFormat, apex::RenderDataFormat::FLOAT3);
+		}
+		if (VertexFormat.mHasStaticColors || VertexFormat.mHasDynamicColors)
+		{
+			vb.setSemanticData( apex::RenderVertexSemantic::COLOR, &Vertices.GetData()->color, sizeof( apex::Vertex ), apex::RenderDataFormat::R8G8B8A8, apex::RenderDataFormat::R32G32B32A32_FLOAT);
+		}
+		for (PxU32 uvNum = 0; uvNum < VertexFormat.mUVCount; ++uvNum)
+		{
+			vb.setSemanticData( (apex::RenderVertexSemantic::Enum)(uvNum+apex::RenderVertexSemantic::TEXCOORD0), &Vertices.GetData()->uv[uvNum], sizeof( apex::Vertex ), apex::RenderDataFormat::FLOAT2);
+		}
+		switch (VertexFormat.mBonesPerVertex)
+		{
+		case 1:
+			vb.setSemanticData( apex::RenderVertexSemantic::BONE_INDEX, &Vertices.GetData()->boneIndices[0], sizeof( apex::Vertex ), apex::RenderDataFormat::USHORT1);
+			break;
+		case 2:
+			vb.setSemanticData( apex::RenderVertexSemantic::BONE_INDEX, &Vertices.GetData()->boneIndices[0], sizeof( apex::Vertex ), apex::RenderDataFormat::USHORT2);
+			vb.setSemanticData( apex::RenderVertexSemantic::BONE_WEIGHT, &Vertices.GetData()->boneWeights[0], sizeof( apex::Vertex ), apex::RenderDataFormat::FLOAT2);
+			break;
+		case 3:
+			vb.setSemanticData( apex::RenderVertexSemantic::BONE_INDEX, &Vertices.GetData()->boneIndices[0], sizeof( apex::Vertex ), apex::RenderDataFormat::USHORT3);
+			vb.setSemanticData( apex::RenderVertexSemantic::BONE_WEIGHT, &Vertices.GetData()->boneWeights[0], sizeof( apex::Vertex ), apex::RenderDataFormat::FLOAT3);
+			break;
+		case 4:
+			vb.setSemanticData( apex::RenderVertexSemantic::BONE_INDEX, &Vertices.GetData()->boneIndices[0], sizeof( apex::Vertex ), apex::RenderDataFormat::USHORT4);
+			vb.setSemanticData( apex::RenderVertexSemantic::BONE_WEIGHT, &Vertices.GetData()->boneWeights[0], sizeof( apex::Vertex ), apex::RenderDataFormat::FLOAT4);
+			break;
+		}
+	}
+
+	// BRG - this is being done until I add the ability to reset the submesh data without resetting the mesh data in the fracture tools
+	while (MeshDesc.m_numSubmeshes > 0)
+	{
+		if (MeshDesc.m_submeshes[MeshDesc.m_numSubmeshes-1].m_vertexIndices != 0)
+		{
+			break;
+		}
+		--MeshDesc.m_numSubmeshes;
+	}
+
+	RenderMeshAssetAuthor.createRenderMesh(MeshDesc, false);
+}
+
+class FProgressListener : public IProgressListener
+{
+public:
+	// Begin IProgressListener interface
+	virtual void	setProgress(int progress, const char* taskName = NULL) override {}
+	// End IProgressListener interface
+};
+
+class FExplicitHierarchicalMeshEmbedding : public apex::ExplicitHierarchicalMesh::Embedding
+{
+public:
+	virtual void	serialize(physx::general_PxIOStream2::PxFileBuf& stream, apex::ExplicitHierarchicalMesh::Embedding::DataType type) const override {}
+	virtual void	deserialize(physx::general_PxIOStream2::PxFileBuf& stream, apex::ExplicitHierarchicalMesh::Embedding::DataType type, physx::PxU32 version) override {}
+};
+
+#endif // WITH_APEX
+#endif // #if WITH_EDITOR	//	Fracture code is only needed in editor
+
+//////////////////////////////////////////////////////////////////////////
+// FFractureMaterial
+//////////////////////////////////////////////////////////////////////////
+
+#if WITH_APEX
+void FFractureMaterial::FillNxFractureMaterialDesc(apex::FractureMaterialDesc& PFractureMaterialDesc)
+{
+#if WITH_EDITOR	//	Fracture code is only needed in editor
+	PFractureMaterialDesc.uvScale = PxVec2(UVScale.X, UVScale.Y);
+	PFractureMaterialDesc.uvOffset = PxVec2(UVOffset.X, UVOffset.Y);
+	PFractureMaterialDesc.tangent = U2PVector(Tangent);
+	PFractureMaterialDesc.uAngle = UAngle;
+	PFractureMaterialDesc.interiorSubmeshIndex = InteriorElementIndex >= 0 ? (PxU32)InteriorElementIndex : 0xFFFFFFFF;	// We'll use this value to indicate we should create a new element
+#endif // WITH_EDITOR
+}
+#endif // WITH_APEX
+
+//////////////////////////////////////////////////////////////////////////
+// UDestructibleFractureSettings
+//////////////////////////////////////////////////////////////////////////
+PRAGMA_DISABLE_DEPRECATION_WARNINGS
+UDestructibleFractureSettings::UDestructibleFractureSettings(const FObjectInitializer& ObjectInitializer)
+	: Super(ObjectInitializer)
+{
+	CellSiteCount = 25;
+	OriginalSubmeshCount = 0;
+}
+
+UDestructibleFractureSettings::UDestructibleFractureSettings(FVTableHelper& Helper)
+	: Super(Helper)
+{
+
+}
+PRAGMA_ENABLE_DEPRECATION_WARNINGS
+
+void UDestructibleFractureSettings::Serialize(FArchive& Ar)
+{
+	Super::Serialize(Ar);
+
+#if WITH_EDITOR	//	Fracture code is only needed in editor
+
+	if (Ar.IsLoading())
+	{
+		// Buffer for the NxDestructibleAsset(Authoring)
+		TArray<uint8> Buffer;
+		uint32 Size;
+		Ar << Size;
+		if (Size > 0)
+		{
+			// Size is non-zero, so a binary blob follows
+			Buffer.AddUninitialized(Size);
+			Ar.Serialize(Buffer.GetData(), Size);
+#if WITH_APEX
+			check(ApexDestructibleAssetAuthoring != NULL);
+			// Wrap this blob with the APEX read stream class
+			physx::PxFileBuf* Stream = GApexSDK->createMemoryReadStream(Buffer.GetData(), Size);
+			if (Stream != NULL)
+			{
+				// Now deserialize the ExplicitHierarchicalMeshes
+				FExplicitHierarchicalMeshEmbedding ExplicitHierarchicalMeshEmbedding;
+				ApexDestructibleAssetAuthoring->getExplicitHierarchicalMesh().deserialize(*Stream, ExplicitHierarchicalMeshEmbedding);
+				ApexDestructibleAssetAuthoring->getCoreExplicitHierarchicalMesh().deserialize(*Stream, ExplicitHierarchicalMeshEmbedding);
+				// Release the stream
+				GApexSDK->releaseMemoryReadStream(*Stream);
+			}
+#endif
+		}
+	}
+	else if (Ar.IsSaving())
+	{
+		uint32 size=0;
+		Ar << size;
+	}
+
+#endif // #if WITH_EDITOR	//	Fracture code is only needed in editor
+}
+
+//IMPLEMENT_FSTRUCTUREDARCHIVE_SERIALIZER(UDestructibleFractureSettings);
+
+void UDestructibleFractureSettings::PostInitProperties()
+{
+	Super::PostInitProperties();
+
+#if WITH_EDITOR	//	Fracture code is only needed in editor
+#if WITH_APEX
+	if(!HasAnyFlags(RF_ClassDefaultObject))
+	{
+		ApexDestructibleAssetAuthoring = static_cast<apex::DestructibleAssetAuthoring*>(apex::GetApexSDK()->createAssetAuthoring(DESTRUCTIBLE_AUTHORING_TYPE_NAME));
+	}
+#endif // WITH_APEX
+#endif // #if WITH_EDITOR	//	Fracture code is only needed in editor
+}
+
+void UDestructibleFractureSettings::BeginDestroy()
+{
+#if WITH_EDITOR	//	Fracture code is only needed in editor
+#if WITH_APEX
+	if (ApexDestructibleAssetAuthoring != NULL)
+	{
+		GPhysCommandHandler->DeferredRelease(ApexDestructibleAssetAuthoring);
+		ApexDestructibleAssetAuthoring = NULL;
+	}
+#endif // WITH_APEX
+#endif // #if WITH_EDITOR	//	Fracture code is only needed in editor
+
+	Super::BeginDestroy();
+}
+
+#if WITH_APEX && WITH_EDITOR
+void UDestructibleFractureSettings::BuildDestructibleAssetCookingDesc(apex::DestructibleAssetCookingDesc& DestructibleAssetCookingDesc)
+{
+	//	Fracture code is only needed in editor
+	// Retrieve the authoring mesh
+	apex::ExplicitHierarchicalMesh& HMesh = ApexDestructibleAssetAuthoring->getExplicitHierarchicalMesh();
+
+	if (ChunkParameters.Num() < (int32)HMesh.chunkCount())
+	{
+		const int32 OldParametersCount = ChunkParameters.Num();
+		ChunkParameters.InsertZeroed(ChunkParameters.Num(), HMesh.chunkCount()-ChunkParameters.Num());
+		for (int32 ChunkIndex = OldParametersCount; ChunkIndex < ChunkParameters.Num(); ++ChunkIndex)
+		{
+PRAGMA_DISABLE_DEPRECATION_WARNINGS
+			ChunkParameters[ChunkIndex] = FDestructibleChunkParameters();
+PRAGMA_ENABLE_DEPRECATION_WARNINGS
+		}
+	}
+
+	// Set up chunk desc array
+	ChunkDescs.Init(apex::DestructibleChunkDesc(), HMesh.chunkCount());
+	DestructibleAssetCookingDesc.chunkDescs = ChunkDescs.GetData();
+	DestructibleAssetCookingDesc.chunkDescCount = ChunkDescs.Num();
+	for (uint32 ChunkIndex = 0; ChunkIndex < DestructibleAssetCookingDesc.chunkDescCount; ++ChunkIndex)
+	{
+PRAGMA_DISABLE_DEPRECATION_WARNINGS
+		FDestructibleChunkParameters& IndexedChunkParameters = ChunkParameters[ChunkIndex];
+		DestructibleAssetCookingDesc.chunkDescs[ChunkIndex].setToDefault();
+		DestructibleAssetCookingDesc.chunkDescs[ChunkIndex].parentIndex = *HMesh.parentIndex(ChunkIndex);
+		DestructibleAssetCookingDesc.chunkDescs[ChunkIndex].meshIndex = *HMesh.partIndex(ChunkIndex);
+		DestructibleAssetCookingDesc.chunkDescs[ChunkIndex].isSupportChunk = IndexedChunkParameters.bIsSupportChunk;
+		DestructibleAssetCookingDesc.chunkDescs[ChunkIndex].doNotFracture = IndexedChunkParameters.bDoNotFracture;
+		DestructibleAssetCookingDesc.chunkDescs[ChunkIndex].doNotDamage = IndexedChunkParameters.bDoNotDamage;
+		DestructibleAssetCookingDesc.chunkDescs[ChunkIndex].doNotCrumble = IndexedChunkParameters.bDoNotCrumble;
+		const bool instancing = (*HMesh.chunkFlags(ChunkIndex) & apex::DestructibleAsset::ChunkIsInstanced) != 0;
+PRAGMA_ENABLE_DEPRECATION_WARNINGS
+		DestructibleAssetCookingDesc.chunkDescs[ChunkIndex].useInstancedRendering = instancing;
+		if (instancing)
+		{
+			DestructibleAssetCookingDesc.chunkDescs[ChunkIndex].instancePositionOffset = *HMesh.instancedPositionOffset(ChunkIndex);
+			DestructibleAssetCookingDesc.chunkDescs[ChunkIndex].instanceUVOffset = *HMesh.instancedUVOffset(ChunkIndex);
+		}
+	}
+
+	// Set up geometry desc array
+	GeometryDescs.Init(apex::DestructibleGeometryDesc(), HMesh.partCount());
+	DestructibleAssetCookingDesc.geometryDescs = GeometryDescs.GetData();
+	DestructibleAssetCookingDesc.geometryDescCount = GeometryDescs.Num();
+	for (uint32 GeometryIndex = 0; GeometryIndex < DestructibleAssetCookingDesc.geometryDescCount; ++GeometryIndex)
+	{
+		DestructibleAssetCookingDesc.geometryDescs[GeometryIndex].setToDefault();
+		DestructibleAssetCookingDesc.geometryDescs[GeometryIndex].convexHullCount = HMesh.convexHullCount(GeometryIndex);
+		DestructibleAssetCookingDesc.geometryDescs[GeometryIndex].convexHulls = HMesh.convexHulls(GeometryIndex);
+	}
+}
+
+bool UDestructibleFractureSettings::SetRootMesh(const TArray<apex::ExplicitRenderTriangle>& MeshTriangles, const TArray<UMaterialInterface*>& InMaterials, const TArray<apex::ExplicitSubmeshData>& SubmeshData,
+												const TArray<uint32>& MeshPartition, bool bFirstPartitionIsDepthZero)
+{
+	bool Success = false;
+	physx::PxI32 NegativeOne = -1;
+
+	if (ApexDestructibleAssetAuthoring != NULL)
+	{
+		Success = ApexDestructibleAssetAuthoring->setRootMesh(MeshTriangles.GetData(), MeshTriangles.Num(), SubmeshData.GetData(), 
+															  SubmeshData.Num(), (uint32*)MeshPartition.GetData(), MeshPartition.Num(), &NegativeOne, bFirstPartitionIsDepthZero ? 1 : 0);
+		if (Success)
+		{
+			apex::CollisionDesc CollisionDesc;
+			ApexDestructibleAssetAuthoring->getExplicitHierarchicalMesh().buildCollisionGeometryForRootChunkParts(CollisionDesc);
+		}
+
+		// Resize the chunk parameters
+		apex::ExplicitHierarchicalMesh& HMesh = ApexDestructibleAssetAuthoring->getExplicitHierarchicalMesh();
+PRAGMA_DISABLE_DEPRECATION_WARNINGS
+		ChunkParameters.Init(FDestructibleChunkParameters(), HMesh.chunkCount());
+PRAGMA_ENABLE_DEPRECATION_WARNINGS
+	}
+
+	OriginalSubmeshCount = SubmeshData.Num();
+
+	Materials = InMaterials;
+
+	return Success;
+}
+
+bool UDestructibleFractureSettings::BuildRootMeshFromApexDestructibleAsset(apex::DestructibleAsset& ApexDestructibleAsset, EDestructibleImportOptions::Type Options)
+{
+	bool Success = false;
+
+	if (ApexDestructibleAssetAuthoring != NULL)
+	{
+		Success = ApexDestructibleAssetAuthoring->importDestructibleAssetToRootMesh(ApexDestructibleAsset, 0);
+
+		apex::ExplicitHierarchicalMesh& EHM = ApexDestructibleAssetAuthoring->getExplicitHierarchicalMesh();
+
+		if (!(Options & EDestructibleImportOptions::PreserveSettings))
+		{
+			// Now apply the y -> -y and v -> 1-v transformation to all vertex data
+			for (PxU32 PartIndex = 0; PartIndex < EHM.partCount(); ++PartIndex)
+			{
+				apex::ExplicitRenderTriangle* Triangles = EHM.meshTriangles(PartIndex);
+				for (PxU32 TriangleIndex = 0; TriangleIndex < EHM.meshTriangleCount(PartIndex); ++TriangleIndex)
+				{
+					apex::ExplicitRenderTriangle& Triangle = Triangles[TriangleIndex];
+					for (PxU32 VertexNum = 0; VertexNum < 3; ++VertexNum)
+					{
+						apex::Vertex& Vertex = Triangle.vertices[VertexNum];
+						Vertex.position.y *= -1.0f;
+						Vertex.normal.y *= -1.0f;
+						Vertex.tangent.y *= -1.0f;
+						Vertex.binormal.y *= -1.0f;
+						
+						for (PxU32 UVNum = 0; UVNum < apex::VertexFormat::MAX_UV_COUNT; ++UVNum)
+						{
+							Vertex.uv[UVNum].v = 1.0f - Vertex.uv[UVNum].v;
+						}
+					}
+					Swap(Triangle.vertices[0], Triangle.vertices[2]);
+				}
+			}
+
+PRAGMA_DISABLE_DEPRECATION_WARNINGS
+			// BRG - Until APEX 1.2.3 or later is in place, chunk flags will not be imported automaticlally, so we need to do this here:
+			ChunkParameters.Init(FDestructibleChunkParameters(), ApexDestructibleAsset.getChunkCount());
+PRAGMA_ENABLE_DEPRECATION_WARNINGS
+
+			enum ChunkFlags	// Copied from DestructibleAsset.h:
+			{
+				SupportChunk =	(1 << 0),
+				UnfractureableChunk =	(1 << 1),
+				DescendantUnfractureable =	(1 << 2),
+				UndamageableChunk =	(1 << 3),
+				UncrumbleableChunk =	(1 << 4),
+
+				Instanced = (1 << 8),
+			};
+			const NvParameterized::Interface* Params = ApexDestructibleAsset.getAssetNvParameterized();
+			if (Params != NULL)
+			{
+				// Damage parameters
+				for (PxU32 ChunkIndex = 0; ChunkIndex < EHM.chunkCount(); ++ChunkIndex)
+				{
+PRAGMA_DISABLE_DEPRECATION_WARNINGS
+					char ChunkFlagsName[MAX_SPRINTF];
+					FCStringAnsi::Sprintf(ChunkFlagsName, "chunks[%d].flags", ChunkIndex);
+					PxU16 ChunkFlags = 0;
+					verify( NvParameterized::getParamU16(*Params, ChunkFlagsName, ChunkFlags) );
+					FDestructibleChunkParameters& IndexedChunkParameters = ChunkParameters[ChunkIndex];
+					IndexedChunkParameters.bIsSupportChunk = (ChunkFlags & SupportChunk) != 0;
+					IndexedChunkParameters.bDoNotFracture = (ChunkFlags & UnfractureableChunk) != 0;
+					IndexedChunkParameters.bDoNotDamage = (ChunkFlags & UndamageableChunk) != 0;
+					IndexedChunkParameters.bDoNotCrumble = (ChunkFlags & UncrumbleableChunk) != 0;
+PRAGMA_ENABLE_DEPRECATION_WARNINGS
+				}
+			}
+		}
+	}
+
+	return Success;
+}
+
+void UDestructibleFractureSettings::CreateVoronoiSitesInRootMesh()
+{
+	if (ApexDestructibleAssetAuthoring != NULL)
+	{
+		VoronoiSites.Init(FVector(0.0f), CellSiteCount);
+		// Progress listener for reporting progress - for now, just a dummy
+		FProgressListener ProgressListener;
+		check(sizeof(FVector) == sizeof(PxVec3));
+		ApexDestructibleAssetAuthoring->createVoronoiSitesInsideMesh((PxVec3*)VoronoiSites.GetData(), nullptr, VoronoiSites.Num(), (PxU32*)&RandomSeed, NULL, apex::BSPOpenMode::Automatic, ProgressListener);
+	}
+}
+
+bool UDestructibleFractureSettings::VoronoiSplitMesh()
+{
+	bool Success = false;
+
+	if (ApexDestructibleAssetAuthoring != NULL)
+	{
+		// Fill MeshProcessingParameters
+		FractureTools::MeshProcessingParameters FTMeshProcessingParameters;
+		FTMeshProcessingParameters.islandGeneration = false;	// expose
+
+		// Fill Voronoi splitting descriptor
+		FractureTools::FractureVoronoiDesc FTFractureVoronoiDesc;
+		FTFractureVoronoiDesc.siteCount = VoronoiSites.Num();
+		check(sizeof(FVector) == sizeof(PxVec3));
+		FTFractureVoronoiDesc.sites = (PxVec3*)VoronoiSites.GetData();
+
+PRAGMA_DISABLE_DEPRECATION_WARNINGS
+		// Material descriptor
+		FractureMaterialDesc.FillNxFractureMaterialDesc(FTFractureVoronoiDesc.materialDesc);
+PRAGMA_ENABLE_DEPRECATION_WARNINGS
+		// Retrieve the authoring mesh to see if the interiorSubmeshIndex is valid
+		apex::ExplicitHierarchicalMesh& HMesh = ApexDestructibleAssetAuthoring->getExplicitHierarchicalMesh();
+		if (FTFractureVoronoiDesc.materialDesc.interiorSubmeshIndex >= HMesh.submeshCount())
+		{
+			// For now just copy the submesh data from the 0 submesh.
+			apex::ExplicitSubmeshData SubmeshData;
+			apex::ExplicitSubmeshData* SourceSubmeshData = HMesh.submeshData(0);
+			if (SourceSubmeshData != NULL)
+			{
+				SubmeshData = *SourceSubmeshData;
+			}
+
+			// BRG - this is being done until I add the ability to reset the submesh data without resetting the mesh data in the fracture tools
+			if (HMesh.submeshCount() > (uint32)OriginalSubmeshCount)
+			{
+				FTFractureVoronoiDesc.materialDesc.interiorSubmeshIndex = (uint32)OriginalSubmeshCount;
+			}
+			else
+			{
+				FTFractureVoronoiDesc.materialDesc.interiorSubmeshIndex = HMesh.addSubmesh(SubmeshData);
+			}
+
+			// Parallel storage of UE materials
+			while (Materials.Num() <= OriginalSubmeshCount)
+			{
+				Materials.Add(UMaterial::GetDefaultMaterial(MD_Surface));
+			}
+		}
+
+		// Collision volume descriptor
+		apex::CollisionDesc CollisionVolumeDesc;
+
+		// Progress listener for reporting progress - for now, just a dummy
+		FProgressListener ProgressListener;
+		Success = ApexDestructibleAssetAuthoring->createVoronoiSplitMesh(FTMeshProcessingParameters, FTFractureVoronoiDesc, CollisionVolumeDesc, false, 0, RandomSeed, ProgressListener);
+	}
+
+	return Success;
+}
+
+PRAGMA_DISABLE_DEPRECATION_WARNINGS
+apex::DestructibleAsset* UDestructibleFractureSettings::CreateApexDestructibleAsset(const apex::DestructibleAssetCookingDesc& DestructibleAssetCookingDesc)
+{
+	apex::DestructibleAsset* ApexDestructibleAsset = NULL;
+
+	if (ApexDestructibleAssetAuthoring != NULL && DestructibleAssetCookingDesc.isValid())
+	{
+		apex::RenderMeshAssetAuthoring* ApexRenderMeshAssetAuthoring = static_cast<apex::RenderMeshAssetAuthoring*>(apex::GetApexSDK()->createAssetAuthoring(RENDER_MESH_AUTHORING_TYPE_NAME));
+		if (ApexRenderMeshAssetAuthoring != NULL)
+		{
+			BuildApexRenderMesh(*ApexRenderMeshAssetAuthoring, ApexDestructibleAssetAuthoring->getExplicitHierarchicalMesh());
+			apex::RenderMeshAsset* ApexRenderMeshAsset = static_cast<apex::RenderMeshAsset*>(apex::GetApexSDK()->createAsset(ApexRenderMeshAssetAuthoring->releaseAndReturnNvParameterizedInterface(), NULL ));
+			if (ApexRenderMeshAsset != NULL)
+			{
+				ApexDestructibleAssetAuthoring->setRenderMeshAsset(ApexRenderMeshAsset);
+				ApexDestructibleAssetAuthoring->cookChunks(DestructibleAssetCookingDesc);
+				ApexDestructibleAsset = static_cast<apex::DestructibleAsset*>(apex::GetApexSDK()->createAsset(*ApexDestructibleAssetAuthoring, NULL));
+			}
+		}
+	}
+
+	return ApexDestructibleAsset;
+}
+PRAGMA_ENABLE_DEPRECATION_WARNINGS
+
+#endif // WITH_APEX && WITH_EDITOR
\ No newline at end of file
diff --git a/Engine/Plugins/Runtime/ApexDestruction/Source/ApexDestruction/Private/DestructibleMesh.cpp b/Engine/Plugins/Runtime/ApexDestruction/Source/ApexDestruction/Private/DestructibleMesh.cpp
new file mode 100644
--- /dev/null
+++ b/Engine/Plugins/Runtime/ApexDestruction/Source/ApexDestruction/Private/DestructibleMesh.cpp
@@ -0,0 +1,242 @@
+// Copyright Epic Games, Inc. All Rights Reserved.
+
+/*=============================================================================
+	DestructibleMesh.cpp: UDestructibleMesh methods.
+=============================================================================*/
+
+#include "DestructibleMesh.h"
+#include "RawIndexBuffer.h"
+#include "DestructibleFractureSettings.h"
+#include "GPUSkinVertexFactory.h"
+#include "UObject/FrameworkObjectVersion.h"
+#include "PhysicalMaterials/PhysicalMaterial.h"
+#include "StaticMeshResources.h"
+#include "PhysXPublic.h"
+#include "Engine/StaticMesh.h"
+#include "Rendering/SkeletalMeshRenderData.h"
+#include "ApexDestructionModule.h"
+
+#if WITH_EDITOR
+#include "EditorFramework/AssetImportData.h"
+#endif
+
+#if WITH_APEX && WITH_EDITOR
+#include "ApexDestructibleAssetImport.h"
+#endif
+
+DEFINE_LOG_CATEGORY(LogDestructible)
+
+PRAGMA_DISABLE_DEPRECATION_WARNINGS
+UDestructibleMesh::UDestructibleMesh(const FObjectInitializer& ObjectInitializer)
+	: Super(ObjectInitializer)
+{
+}
+PRAGMA_ENABLE_DEPRECATION_WARNINGS
+
+void UDestructibleMesh::PostLoad()
+{
+	Super::PostLoad();
+
+#if WITH_DESTRUCTIBLE_DEPRECATION
+	// Deprecation message as of UE5.1
+	UE_LOG(LogDestructible, Warning, TEXT("During PostLoad for Mesh %s: DestructibleMesh is deprecated. Destruction is now supported by UGeometryCollection. Please update your assets before upgrading to the next release."), *GetPathName());
+#endif
+
+	// BodySetup is used for uniform lookup of PhysicalMaterials.
+	CreateBodySetup();
+}
+
+#if WITH_EDITOR
+void UDestructibleMesh::PreEditChange(FProperty* PropertyAboutToChange)
+{
+	Super::PreEditChange(PropertyAboutToChange);
+
+	CreateBodySetup();
+}
+#endif
+
+float ImpactResistanceToAPEX(bool bCustomResistance, float ImpactResistance)
+{
+	return bCustomResistance ? ImpactResistance : 0.f;
+}
+
+
+void APEXToImpactResistance(bool& bCustomImpactResistance, float& ImpactResistance)
+{
+	//apex interprets 0 as disabled, but we want custom flag
+	bCustomImpactResistance = ImpactResistance != 0.f;
+	if (ImpactResistance == 0.f)
+	{
+		ImpactResistance = 1.f;
+	}
+}
+
+int32 DefaultImpactDamageDepthToAPEX(bool bEnableImpactDamage, int32 DefaultImpactDamageDepth)
+{
+	return bEnableImpactDamage ? DefaultImpactDamageDepth : -1.f;
+}
+
+void APEXToDefaultImpactDamageDepth(bool& bEnableImpactDamage, int32& DefaultImpactDamageDepth)
+{
+	//apex interprets -1 as disabled, but we want custom flag
+	bEnableImpactDamage = DefaultImpactDamageDepth != -1;
+	DefaultImpactDamageDepth = 1;
+}
+
+int32 DebrisDepthToAPEX(bool bEnableDebris, int32 DebrisDepth)
+{
+	return bEnableDebris ? DebrisDepth : -1;
+}
+
+void APEXToDebrisDepth(bool& bEnableDebris, int32& DebrisDepth)
+{
+	bEnableDebris = DebrisDepth != -1;
+	if (DebrisDepth == -1)
+	{
+		DebrisDepth = 0;
+	}
+}
+
+void UDestructibleMesh::Serialize(FArchive& Ar)
+{
+	Super::Serialize( Ar );
+
+	Ar.UsingCustomVersion(FFrameworkObjectVersion::GUID);
+
+	if( Ar.IsLoading() )
+	{
+		// Deserializing the name of the NxDestructibleAsset
+		TArray<uint8> NameBuffer;
+		uint32 NameBufferSize;
+		Ar << NameBufferSize;
+		NameBuffer.AddUninitialized( NameBufferSize );
+		Ar.Serialize( NameBuffer.GetData(), NameBufferSize );
+
+		// Buffer for the NxDestructibleAsset
+		TArray<uint8> Buffer;
+		uint32 Size;
+		Ar << Size;
+		if( Size > 0 )
+		{
+			// Size is non-zero, so a binary blob follows
+			// Discard apex serialized buffer as no longer supported
+			Buffer.AddUninitialized( Size );
+			Ar.Serialize( Buffer.GetData(), Size );
+		}
+		if (Ar.CustomVer(FFrameworkObjectVersion::GUID) >= FFrameworkObjectVersion::CacheDestructibleOverlaps)
+		{
+			Ar << Size;
+			if( Size > 0 )
+			{
+				// Here comes the collision data cache
+				Buffer.Reset(Size);
+				Buffer.AddZeroed(Size);
+
+				Ar.Serialize( Buffer.GetData(), Size );
+			}
+		}
+	}
+	else if ( Ar.IsSaving() )
+	{
+		const char* Name = "NO_APEX";
+		// Serialize the name
+		uint32 NameBufferSize = FCStringAnsi::Strlen( Name )+1;
+		Ar << NameBufferSize;
+		Ar.Serialize( (void*)Name, NameBufferSize );
+
+		uint32 size=0;
+		Ar << size; // Buffer for the NxDestructibleAsset
+		Ar << size; // collision data cache
+	}
+
+	if (Ar.UEVer() < VER_UE4_CLEAN_DESTRUCTIBLE_SETTINGS)
+	{
+PRAGMA_DISABLE_DEPRECATION_WARNINGS
+		APEXToImpactResistance(DefaultDestructibleParameters.DamageParameters.bCustomImpactResistance, DefaultDestructibleParameters.DamageParameters.ImpactResistance);
+		APEXToDefaultImpactDamageDepth(DefaultDestructibleParameters.DamageParameters.bEnableImpactDamage, DefaultDestructibleParameters.DamageParameters.DefaultImpactDamageDepth);
+		APEXToDebrisDepth(DefaultDestructibleParameters.SpecialHierarchyDepths.bEnableDebris, DefaultDestructibleParameters.SpecialHierarchyDepths.DebrisDepth);
+PRAGMA_ENABLE_DEPRECATION_WARNINGS
+	}
+}
+
+void UDestructibleMesh::FinishDestroy()
+{
+	Super::FinishDestroy();
+}
+
+void UDestructibleMesh::LoadDefaultDestructibleParametersFromApexAsset()
+{
+
+}
+
+void UDestructibleMesh::CreateFractureSettings()
+{
+PRAGMA_DISABLE_DEPRECATION_WARNINGS
+#if WITH_EDITORONLY_DATA
+	if (FractureSettings == NULL)
+	{
+		FractureSettings = NewObject<UDestructibleFractureSettings>(this);
+		check(FractureSettings);
+	}
+#endif	// WITH_EDITORONLY_DATA
+PRAGMA_ENABLE_DEPRECATION_WARNINGS
+}
+
+bool UDestructibleMesh::BuildFractureSettingsFromStaticMesh(UStaticMesh* StaticMesh)
+{
+	return false;
+}
+
+APEXDESTRUCTION_API bool UDestructibleMesh::BuildFromStaticMesh( UStaticMesh& StaticMesh )
+{
+#if WITH_EDITOR
+	PreEditChange(NULL);
+
+	// Import the StaticMesh
+	if (!BuildFractureSettingsFromStaticMesh(&StaticMesh))
+	{
+		return false;
+	}
+
+	SourceStaticMesh = &StaticMesh;
+
+	SourceSMImportTimestamp = FDateTime::MinValue();
+	if ( StaticMesh.AssetImportData && StaticMesh.AssetImportData->SourceData.SourceFiles.Num() == 1 )
+	{
+		SourceSMImportTimestamp = StaticMesh.AssetImportData->SourceData.SourceFiles[0].Timestamp;
+	}
+
+	PostEditChange();
+	MarkPackageDirty();
+#endif // WITH_EDITOR
+	return true;
+}
+
+APEXDESTRUCTION_API bool UDestructibleMesh::SetupChunksFromStaticMeshes( const TArray<UStaticMesh*>& ChunkMeshes )
+{
+#if WITH_EDITOR
+	if (SourceStaticMesh == NULL)
+	{
+		UE_LOG(LogDestructible, Warning, TEXT("Unable to import FBX as level 1 chunks if the DM was not created from a static mesh."));
+		return false;
+	}
+
+	PreEditChange(NULL);
+
+	FractureChunkMeshes.Empty(ChunkMeshes.Num());
+	FractureChunkMeshes.Append(ChunkMeshes);
+	
+	// Import the StaticMesh
+	if (!BuildFractureSettingsFromStaticMesh(SourceStaticMesh))
+	{
+		return false;
+	}
+
+	// Clear the fracture chunk meshes again
+	FractureChunkMeshes.Empty();
+
+	PostEditChange();
+	MarkPackageDirty();
+#endif // WITH_EDITOR
+	return true;
+}
diff --git a/Engine/Plugins/Runtime/ApexDestruction/Source/ApexDestruction/Public/ApexDestructibleAssetImport.h b/Engine/Plugins/Runtime/ApexDestruction/Source/ApexDestruction/Public/ApexDestructibleAssetImport.h
new file mode 100644
--- /dev/null
+++ b/Engine/Plugins/Runtime/ApexDestruction/Source/ApexDestruction/Public/ApexDestructibleAssetImport.h
@@ -0,0 +1,112 @@
+// Copyright Epic Games, Inc. All Rights Reserved.
+#pragma once
+
+#include "CoreMinimal.h"
+#include "EngineDefines.h"
+#include "DestructibleFractureSettings.h"
+
+#if WITH_EDITOR
+
+class FSkeletalMeshImportData;
+class UDestructibleMesh;
+
+/*=============================================================================
+	ApexDestructibleAssetImport.h:
+
+	Creation of an APEX NxDestructibleAsset from a binary buffer.
+	DestructibleMesh creation from an APEX NxDestructibleAsset.
+
+=============================================================================*/
+
+#if WITH_APEX
+
+// Forward declarations
+class UDestructibleMesh;
+class FSkelMeshOptionalImportData;
+class FSkeletalMeshImportData;
+
+namespace nvidia
+{
+	namespace apex
+	{
+		class DestructibleAsset;
+	};
+};
+
+
+
+/**
+ * Creates an APEX NxDestructibleAsset from a binary buffer.
+ *
+ * @param Buffer - pointer to the beginning of the buffer
+ * @param BufferSize - the size of the buffer in bytes
+ *
+ * @return The newly created NxDestructibleAsset if successful, NULL otherwise
+ */
+UE_DEPRECATED(4.26, "APEX is deprecated. Destruction in future will be supported using Chaos Destruction.")
+APEXDESTRUCTION_API nvidia::apex::DestructibleAsset* CreateApexDestructibleAssetFromBuffer(const uint8* Buffer, int32 BufferSize);
+
+
+/**
+ * Creates an APEX NxDestructibleAsset from a named file.
+ *
+ * @param Filename - name of the file to import
+ *
+ * @return The newly created NxDestructibleAsset if successful, NULL otherwise
+ */
+UE_DEPRECATED(4.26, "APEX is deprecated. Destruction in future will be supported using Chaos Destruction.")
+APEXDESTRUCTION_API nvidia::apex::DestructibleAsset* CreateApexDestructibleAssetFromFile(const FString& Filename);
+
+/**
+ * Sets the NxDestructibleAsset represented by the UDestructibleMesh.  If there was an existing NxDestructibleAsset, its release method is called.
+ * SkeletalMesh data is replaced or created by the render mesh data in the NxDestructibleAsset.
+ *
+ * N.B. The ApexDestructibleAsset will be modified.  Its collision geometry will undergo a reflection along the y-axis to match UE's left-handed convention.
+ * Graphics information will get a similar transformation, and V coordnates will also be flipped using the mapping V' = 1-V, again according to UE's convention.
+ * Not currently implemented, but soon to be, all graphics data in the ApexDestructibleAsset will be deleted after it is translated into the SkeletalMesh data.
+ * 
+ * @param DestructibleMesh - the UDestructibleMesh to which to apply the new APEX Desrtructible Asset
+ * @param ApexDestructibleAsset - the NxDestructibleAsset to use
+ * @param OutData - Optional import data to populate
+ * @param Options - Additional options to modify the import behavior ( see EDestructibleImportOptions::Type )
+ *
+ * @return true if successful
+ */
+PRAGMA_DISABLE_DEPRECATION_WARNINGS
+UE_DEPRECATED(4.26, "APEX is deprecated. Destruction in future will be supported using Chaos Destruction.")
+APEXDESTRUCTION_API bool SetApexDestructibleAsset(UDestructibleMesh& DestructibleMesh, nvidia::apex::DestructibleAsset& ApexDestructibleAsset, FSkeletalMeshImportData* OutData, EDestructibleImportOptions::Type Options = EDestructibleImportOptions::None);
+
+
+/**
+ * Creates a DestructibleMesh from an APEX Destructible Asset with the given name and flags.
+ *
+ * @param InParent - parent object for the UDestructibleMesh
+ * @param ApexDestructibleAsset - the APEX Destructible Asset to import
+ * @param Name - the Unreal name for the UDestructibleMesh
+ * @param Flags - object flags for the UDestructibleMesh
+ * @param OutData - Optional import data to populate
+ * @param Options - Additional options to modify the import behavior ( see EDestructibleImportOptions::Type )
+ *
+ * @return The newly created UDestructibleMesh if successful, NULL otherwise
+ */
+
+UE_DEPRECATED(4.26, "APEX is deprecated. Destruction in future will be supported using Chaos Destruction.")
+APEXDESTRUCTION_API UDestructibleMesh* ImportDestructibleMeshFromApexDestructibleAsset(UObject* InParent, nvidia::apex::DestructibleAsset& ApexDestructibleAsset, FName Name, EObjectFlags Flags, FSkeletalMeshImportData* OutData, EDestructibleImportOptions::Type Options = EDestructibleImportOptions::None);
+
+
+/**
+ * Builds a UDestructibleMesh from its internal fracture settings
+ *
+ * @param DestructibleMesh - the UDestructibleMesh to which to build from its fracture settings
+ * @param OutData - Optional import data to populate
+ *
+ * @return The true iff successful
+ */
+
+UE_DEPRECATED(4.26, "APEX is deprecated. Destruction in future will be supported using Chaos Destruction.")
+APEXDESTRUCTION_API bool BuildDestructibleMeshFromFractureSettings(UDestructibleMesh& DestructibleMesh, FSkeletalMeshImportData* OutData);
+PRAGMA_ENABLE_DEPRECATION_WARNINGS
+
+#endif // WITH_APEX
+
+#endif
\ No newline at end of file
diff --git a/Engine/Plugins/Runtime/ApexDestruction/Source/ApexDestruction/Public/ApexDestructionCustomPayload.h b/Engine/Plugins/Runtime/ApexDestruction/Source/ApexDestruction/Public/ApexDestructionCustomPayload.h
new file mode 100644
--- /dev/null
+++ b/Engine/Plugins/Runtime/ApexDestruction/Source/ApexDestruction/Public/ApexDestructionCustomPayload.h
@@ -0,0 +1,64 @@
+// Copyright Epic Games, Inc. All Rights Reserved.
+
+#pragma once
+
+#include "CoreMinimal.h"
+#include "CustomPhysXPayload.h"
+
+struct UE_DEPRECATED(4.26, "APEX is deprecated. Destruction in future will be supported using Chaos Destruction.") FUpdateChunksInfo
+{
+	int32 ChunkIndex;
+	FTransform WorldTM;
+
+	FUpdateChunksInfo(int32 InChunkIndex, const FTransform & InWorldTM) : ChunkIndex(InChunkIndex), WorldTM(InWorldTM)
+	{}
+};
+
+#if WITH_APEX
+
+class UDestructibleComponent;
+
+struct UE_DEPRECATED(4.26, "APEX is deprecated. Destruction in future will be supported using Chaos Destruction.") FApexDestructionSyncActors : public FCustomPhysXSyncActors
+{
+	virtual void BuildSyncData_AssumesLocked(const TArray<physx::PxRigidActor*> & RigidActors) override;
+
+	virtual void FinalizeSync() override;
+
+private:
+PRAGMA_DISABLE_DEPRECATION_WARNINGS
+	/** Sync data for updating physx sim result */
+	TMap<TWeakObjectPtr<UDestructibleComponent>, TArray<FUpdateChunksInfo> > ComponentUpdateMapping;
+PRAGMA_ENABLE_DEPRECATION_WARNINGS
+};
+
+struct UE_DEPRECATED(4.26, "APEX is deprecated. Destruction in future will be supported using Chaos Destruction.") FApexDestructionCustomPayload : public FCustomPhysXPayload
+{
+PRAGMA_DISABLE_DEPRECATION_WARNINGS
+	FApexDestructionCustomPayload()
+		: FCustomPhysXPayload(SingletonCustomSync)
+	{
+	}
+PRAGMA_ENABLE_DEPRECATION_WARNINGS
+
+	virtual TWeakObjectPtr<UPrimitiveComponent> GetOwningComponent() const override;
+
+	virtual int32 GetItemIndex() const override;
+
+	virtual FName GetBoneName() const override;
+
+	virtual FBodyInstance* GetBodyInstance() const override;
+
+	/** Index of the chunk this data belongs to*/
+	int32 ChunkIndex;
+	/** Component owning this chunk info*/
+	TWeakObjectPtr<UDestructibleComponent> OwningComponent;
+
+private:
+	friend class FApexDestructionModule;
+
+PRAGMA_DISABLE_DEPRECATION_WARNINGS
+	static FApexDestructionSyncActors* SingletonCustomSync;
+PRAGMA_ENABLE_DEPRECATION_WARNINGS
+};
+
+#endif // WITH_APEX
\ No newline at end of file
diff --git a/Engine/Plugins/Runtime/ApexDestruction/Source/ApexDestruction/Public/ApexDestructionModule.h b/Engine/Plugins/Runtime/ApexDestruction/Source/ApexDestruction/Public/ApexDestructionModule.h
new file mode 100644
--- /dev/null
+++ b/Engine/Plugins/Runtime/ApexDestruction/Source/ApexDestruction/Public/ApexDestructionModule.h
@@ -0,0 +1,17 @@
+// Copyright Epic Games, Inc. All Rights Reserved.
+#pragma once
+
+#include "CoreMinimal.h"
+#include "PhysXPublic.h"
+
+namespace nvidia
+{
+
+namespace apex
+{
+	class ModuleDestructible;
+	struct DamageEventReportData;
+}
+}
+
+extern nvidia::apex::ModuleDestructible* GApexModuleDestructible;
\ No newline at end of file
diff --git a/Engine/Plugins/Runtime/ApexDestruction/Source/ApexDestruction/Public/DestructibleActor.h b/Engine/Plugins/Runtime/ApexDestruction/Source/ApexDestruction/Public/DestructibleActor.h
new file mode 100644
--- /dev/null
+++ b/Engine/Plugins/Runtime/ApexDestruction/Source/ApexDestruction/Public/DestructibleActor.h
@@ -0,0 +1,57 @@
+// Copyright Epic Games, Inc. All Rights Reserved.
+
+#pragma once
+
+/** This class represents an APEX Destructible Actor. */
+
+#include "CoreMinimal.h"
+#include "UObject/ObjectMacros.h"
+#include "GameFramework/Actor.h"
+#include "DestructibleActor.generated.h"
+
+class UDestructibleComponent;
+
+/** Delegate for notification when fracture occurs */
+DECLARE_DYNAMIC_MULTICAST_DELEGATE_TwoParams(FActorFractureSignature, const FVector&, HitPoint, const FVector&, HitDirection);
+
+class UE_DEPRECATED(4.26, "APEX is deprecated. Destruction in future will be supported using Chaos Destruction.") ADestructibleActor;
+UCLASS(hideCategories=(Input), showCategories=("Input|MouseInput", "Input|TouchInput"), ComponentWrapperClass)
+class APEXDESTRUCTION_API ADestructibleActor : public AActor
+{
+	GENERATED_UCLASS_BODY()
+
+	/**
+	 * The component which holds the skinned mesh and physics data for this actor.
+	 */
+private:
+PRAGMA_DISABLE_DEPRECATION_WARNINGS
+	UPROPERTY(VisibleAnywhere, BlueprintReadOnly, Category = Destruction, meta = (ExposeFunctionCategories = "Destruction,Components|Destructible", AllowPrivateAccess = "true"))
+	UDestructibleComponent* DestructibleComponent;
+PRAGMA_ENABLE_DEPRECATION_WARNINGS
+public:
+
+	UPROPERTY(BlueprintAssignable, Category = "Components|Destructible")
+	FActorFractureSignature OnActorFracture;
+
+#if WITH_EDITORONLY_DATA
+	UE_DEPRECATED(4.24, "This property will be removed in future versions. Please use bCanEverAffectNavigation in DestructionComponent.")
+	UPROPERTY(config, BlueprintReadWrite, Category = Navigation, meta = (DeprecationMessage = "Setting the value from Blueprint script does nothing. Please use bCanEverAffectNavigation in DestructionComponent."))
+	uint32 bAffectNavigation : 1;
+#endif // WITH_EDITORONLY_DATA
+
+#if WITH_EDITOR
+	//~ Begin AActor Interface.
+	virtual bool GetReferencedContentObjects( TArray<UObject*>& Objects ) const override;
+	virtual void PostLoad() override;
+	//~ End AActor Interface.
+#endif // WITH_EDITOR
+
+public:
+PRAGMA_DISABLE_DEPRECATION_WARNINGS
+	/** Returns DestructibleComponent subobject **/
+	UDestructibleComponent* GetDestructibleComponent() const { return DestructibleComponent; }
+PRAGMA_ENABLE_DEPRECATION_WARNINGS
+};
+
+
+
diff --git a/Engine/Plugins/Runtime/ApexDestruction/Source/ApexDestruction/Public/DestructibleComponent.h b/Engine/Plugins/Runtime/ApexDestruction/Source/ApexDestruction/Public/DestructibleComponent.h
new file mode 100644
--- /dev/null
+++ b/Engine/Plugins/Runtime/ApexDestruction/Source/ApexDestruction/Public/DestructibleComponent.h
@@ -0,0 +1,271 @@
+// Copyright Epic Games, Inc. All Rights Reserved.
+#pragma once
+
+#include "CoreMinimal.h"
+#include "UObject/ObjectMacros.h"
+#include "Engine/EngineTypes.h"
+#include "EngineDefines.h"
+#include "PhysxUserData.h"
+#include "ApexDestructionCustomPayload.h"
+#include "Components/SkinnedMeshComponent.h"
+#include "DestructibleInterface.h"
+#include "DestructibleComponent.generated.h"
+
+class AController;
+class USkeletalMesh;
+struct FBodyInstance;
+struct FCollisionShape;
+struct FNavigableGeometryExport;
+
+#if WITH_PHYSX
+namespace physx
+{
+	class PxRigidDynamic;
+	class PxRigidActor;
+}
+
+#if WITH_APEX
+namespace nvidia
+{
+	namespace apex
+	{
+		class  DestructibleActor;
+		struct DamageEventReportData;
+		struct ChunkStateEventData;
+	}
+}
+#endif
+#endif // WITH_PHYSX 
+
+/** Delegate for notification when fracture occurs */
+DECLARE_DYNAMIC_MULTICAST_DELEGATE_TwoParams(FComponentFractureSignature, const FVector &, HitPoint, const FVector &, HitDirection);
+
+/**
+ *	This component holds the physics data for a DestructibleActor.
+ *
+ *	The USkeletalMesh pointer in the base class (SkinnedMeshComponent) MUST be a DestructibleMesh
+ */
+
+class UE_DEPRECATED(4.26, "APEX is deprecated. Destruction in future will be supported using Chaos Destruction.") UDestructibleComponent;
+UCLASS(ClassGroup = Physics, hidecategories = (Object, Mesh, "Components|SkinnedMesh", Mirroring, Activation, "Components|Activation"), config = Engine, editinlinenew, meta = (BlueprintSpawnableComponent))
+class APEXDESTRUCTION_API UDestructibleComponent : public USkinnedMeshComponent, public IDestructibleInterface
+{
+	GENERATED_UCLASS_BODY()
+
+	/** If set, use this actor's fracture effects instead of the asset's fracture effects. */
+	UPROPERTY(EditAnywhere, BlueprintReadWrite, Category=DestructibleComponent)
+	uint32 bFractureEffectOverride:1;
+
+	/** Fracture effects for each fracture level. Used only if Fracture Effect Override is set. */
+	UPROPERTY(EditAnywhere, BlueprintReadOnly, editfixedsize, Category=DestructibleComponent)
+	TArray<struct FFractureEffect> FractureEffects;
+
+	/**
+	 *	Enable "hard sleeping" for destruction-generated PxActors.  This means that they turn kinematic
+	 *	when they sleep, but can be made dynamic again by application of enough damage.
+	 */	
+	UPROPERTY(EditAnywhere, BlueprintReadWrite, Category=DestructibleComponent)
+	bool bEnableHardSleeping;
+
+	/**
+	 * The minimum size required to treat chunks as Large.
+	 */
+	UPROPERTY(EditAnywhere, BlueprintReadWrite, Category = DestructibleComponent)
+	float LargeChunkThreshold;
+
+#if WITH_EDITORONLY_DATA
+	/** Provide a blueprint interface for setting the destructible mesh */
+	UPROPERTY(Transient, BlueprintSetter=SetDestructibleMesh, BlueprintGetter=GetDestructibleMesh, Category = "Components|Destructible", meta = (DeprecatedProperty, DeprecatedMessage = "Destructible mesh is now deprected, use UGeometryCollection instead"))
+	class UDestructibleMesh* DestructibleMesh_DEPRECATED;
+#endif // WITH_EDITORONLY_DATA
+
+#if WITH_APEX
+PRAGMA_DISABLE_DEPRECATION_WARNINGS
+	/** Per chunk info */
+	TArray<FApexDestructionCustomPayload> ChunkInfos;
+PRAGMA_ENABLE_DEPRECATION_WARNINGS
+#endif // WITH_PHYSX 
+
+#if WITH_EDITOR
+	//~ Begin UObject Interface.
+	virtual void PostEditChangeProperty( struct FPropertyChangedEvent& PropertyChangedEvent) override;
+#endif // WITH_EDITOR
+#if WITH_EDITORONLY_DATA
+	virtual void Serialize(FArchive& Ar) override;
+#endif
+	//~ End UObject Interface.
+
+	// Take damage
+	UFUNCTION(BlueprintCallable, Category="Components|Destructible")
+	virtual void ApplyDamage(float DamageAmount, const FVector& HitLocation, const FVector& ImpulseDir, float ImpulseStrength) override;
+
+	// Take radius damage
+	UFUNCTION(BlueprintCallable, Category="Components|Destructible")
+	virtual void ApplyRadiusDamage(float BaseDamage, const FVector& HurtOrigin, float DamageRadius, float ImpulseStrength, bool bFullDamage) override;
+
+	UE_DEPRECATED("5.1", "Destructible mesh is now deprected, use UGeometryCollection instead")
+	UFUNCTION(BlueprintCallable, BlueprintSetter, Category="Components|Destructible", meta=(DeprecatedFunction, DeprecationMessage="Destructible mesh is now deprected, use UGeometryCollection instead"))
+	void SetDestructibleMesh(class UDestructibleMesh* NewMesh);
+
+	UE_DEPRECATED("5.1", "Destructible mesh is now deprected, use UGeometryCollection instead")
+	UFUNCTION(BlueprintCallable, BlueprintGetter, Category="Components|Destructible", meta = (DeprecatedFunction, DeprecationMessage = "Destructible mesh is now deprected, use UGeometryCollection instead"))
+	class UDestructibleMesh* GetDestructibleMesh();
+
+	/** Called when a component is touched */
+	UPROPERTY(BlueprintAssignable, Category = "Components|Destructible")
+	FComponentFractureSignature OnComponentFracture;
+
+public:
+#if WITH_APEX
+	/** The DestructibleActor instantated from a DestructibleAsset, which contains the runtime physical state. */
+	nvidia::apex::DestructibleActor* ApexDestructibleActor;
+#endif	//WITH_APEX
+
+	//~ Begin USceneComponent Interface.
+	virtual FBoxSphereBounds CalcBounds(const FTransform& LocalToWorld) const override;
+	virtual void OnUpdateTransform(EUpdateTransformFlags UpdateTransformFlags, ETeleportType Teleport = ETeleportType::None) override;
+	virtual void Activate(bool bReset=false) override;
+	virtual void Deactivate() override;
+	//~ End USceneComponent Interface.
+
+	//~ Begin UActorComponent Interface.
+	virtual void BeginPlay() override;
+protected:
+	virtual void OnCreatePhysicsState() override;
+	virtual void OnDestroyPhysicsState() override;
+public:
+	virtual class UBodySetup* GetBodySetup() override;
+	//~ End UActorComponent Interface.
+
+	//~ Begin UPrimitiveComponent Interface.
+	virtual FBodyInstance* GetBodyInstance(FName BoneName = NAME_None, bool bGetWelded = true, int32 Index = INDEX_NONE) const override;
+	virtual bool CanEditSimulatePhysics() override;
+	virtual bool IsAnySimulatingPhysics() const override;
+
+	virtual void AddImpulse(FVector Impulse, FName BoneName = NAME_None, bool bVelChange = false) override;
+	virtual void AddImpulseAtLocation(FVector Impulse, FVector Position, FName BoneName = NAME_None) override;
+	virtual void AddForce(FVector Force, FName BoneName = NAME_None, bool bAccelChange = false) override;
+	virtual void AddForceAtLocation(FVector Force, FVector Location, FName BoneName = NAME_None) override;
+	virtual void AddForceAtLocationLocal(FVector Force, FVector Location, FName BoneName = NAME_None) override;
+	virtual void AddRadialImpulse(FVector Origin, float Radius, float Strength, ERadialImpulseFalloff Falloff, bool bVelChange=false) override;
+	virtual void AddRadialForce(FVector Origin, float Radius, float Strength, ERadialImpulseFalloff Falloff, bool bAccelChange = false) override;
+	virtual void ReceiveComponentDamage(float DamageAmount, FDamageEvent const& DamageEvent, AController* EventInstigator, AActor* DamageCauser) override;
+
+	virtual bool LineTraceComponent( FHitResult& OutHit, const FVector Start, const FVector End, const FCollisionQueryParams& Params ) override;
+    virtual bool SweepComponent( FHitResult& OutHit, const FVector Start, const FVector End, const FQuat& ShapeWorldRotation, const FCollisionShape& CollisionShape, bool bTraceComplex=false) override;
+	virtual void SetEnableGravity(bool bGravityEnabled) override;
+
+	virtual void WakeRigidBody(FName BoneName /* = NAME_None */) override;
+	virtual void SetSimulatePhysics(bool bSimulate) override;
+
+	virtual void SetMaterial(int32 ElementIndex, UMaterialInterface* Material) override;
+	virtual void SetCollisionEnabled(ECollisionEnabled::Type NewType) override;
+    virtual void SetCollisionProfileName(FName InCollisionProfileName, bool bUpdateOverlaps=true) override;
+	virtual void OnActorEnableCollisionChanged() override;
+
+	virtual void SetCollisionResponseToChannel(ECollisionChannel Channel, ECollisionResponse NewResponse) override;
+	virtual void SetCollisionResponseToAllChannels(ECollisionResponse NewResponse) override;
+	virtual void SetCollisionResponseToChannels(const FCollisionResponseContainer& NewReponses) override;
+	//~ End UPrimitiveComponent Interface.
+
+	//~ Begin SkinnedMeshComponent Interface.
+	virtual bool ShouldUpdateTransform(bool bLODHasChanged) const override;
+	virtual void RefreshBoneTransforms(FActorComponentTickFunction* TickFunction = NULL) override;
+	virtual void SetSkeletalMesh(USkeletalMesh* InSkelMesh, bool bReinitPose = true) override;
+	virtual FTransform GetSocketTransform(FName InSocketName, ERelativeTransformSpace TransformSpace = RTS_World) const override;
+	//~ End SkinnedMeshComponent Interface.
+
+
+	//~ Begin DestructibleComponent Interface.
+#if WITH_APEX
+	struct FFakeBodyInstanceState
+	{
+		physx::PxRigidActor* ActorSync;
+		int32 InstanceIndex;
+	};
+
+	/** Changes the body instance to have the specified actor and instance id. */
+	void SetupFakeBodyInstance(physx::PxRigidActor* NewRigidActor, int32 InstanceIdx, FFakeBodyInstanceState* PrevState = NULL);
+	
+	/** Resets the BodyInstance to the state that is defined in PrevState. */
+	void ResetFakeBodyInstance(FFakeBodyInstanceState& PrevState);
+
+	/** Setup a pair of PxShape and ChunkIndex */
+	void Pair( int32 ChunkIndex, physx::PxShape* PShape );
+#endif // WITH_APEX
+
+	/** This method makes a chunk (fractured piece) visible or invisible.
+	 *
+	 * @param ChunkIndex - Which chunk to affect.  ChunkIndex must lie in the range: 0 <= ChunkIndex < ((DestructibleMesh*)USkeletalMesh)->ApexDestructibleAsset->chunkCount().
+	 * @param bVisible - If true, the chunk will be made visible.  Otherwise, the chunk is made invisible.
+	 */
+	void SetChunkVisible( int32 ChunkIndex, bool bInVisible );
+
+
+	/** This method sets a chunk's (fractured piece's) world rotation and translation.
+	 *
+	 * @param ChunkIndex - Which chunk to affect.  ChunkIndex must lie in the range: 0 <= ChunkIndex < ((DestructibleMesh*)USkeletalMesh)->ApexDestructibleAsset->chunkCount().
+	 * @param WorldRotation - The orientation to give to the chunk in world space, represented as a quaternion.
+	 * @param WorldRotation - The world space position to give to the chunk.
+	 */
+	void SetChunkWorldRT( int32 ChunkIndex, const FQuat& WorldRotation, const FVector& WorldTranslation );
+
+#if WITH_APEX
+	/** Trigger any fracture effects after a damage event is received */
+	virtual void SpawnFractureEffectsFromDamageEvent(const nvidia::apex::DamageEventReportData& InDamageEvent);
+
+	/** Callback from physics system to notify the actor that it has been damaged */
+	void OnDamageEvent(const nvidia::apex::DamageEventReportData& InDamageEvent);
+
+	/** Callback from physics system to notify the actor that a chunk's visibility has changed */
+	void OnVisibilityEvent(const nvidia::apex::ChunkStateEventData & InDamageEvent);
+#endif // WITH_APEX
+
+	//~ End DestructibleComponent Interface.
+
+	virtual bool DoCustomNavigableGeometryExport(FNavigableGeometryExport& GeomExport) const override;
+
+	FORCEINLINE static int32 ChunkIdxToBoneIdx(int32 ChunkIdx) { return ChunkIdx + 1; }
+	FORCEINLINE static int32 BoneIdxToChunkIdx(int32 BoneIdx) { return FMath::Max(BoneIdx - 1, 0); }
+private:
+
+PRAGMA_DISABLE_DEPRECATION_WARNINGS
+	void SetChunksWorldTM(const TArray<FUpdateChunksInfo>& UpdateInfos);
+PRAGMA_ENABLE_DEPRECATION_WARNINGS
+
+	bool IsFracturedOrInitiallyStatic() const;
+
+	/** Obtains the appropriate PhysX scene lock for READING and executes the passed in lambda. */
+	bool ExecuteOnPhysicsReadOnly(TFunctionRef<void()> Func) const;
+
+	/** Obtains the appropriate PhysX scene lock for WRITING and executes the passed in lambda. */
+	bool ExecuteOnPhysicsReadWrite(TFunctionRef<void()> Func) const;
+
+	/** Collision response used for chunks */
+	FCollisionResponse LargeChunkCollisionResponse;
+	FCollisionResponse SmallChunkCollisionResponse;
+#if WITH_PHYSX
+	/** User data wrapper for this component passed to physx */
+	FPhysxUserData PhysxUserData;
+
+	void SetCollisionResponseForShape(physx::PxShape* Shape, int32 ChunkIdx);
+	void SetCollisionResponseForActor(physx::PxRigidDynamic* Actor, int32 ChunkIdx, const FCollisionResponseContainer* ResponseOverride = NULL);
+
+
+public:
+	/** User data wrapper for the chunks passed to physx */
+	TArray<FPhysxUserData> PhysxChunkUserData;
+	bool IsChunkLarge(physx::PxRigidActor* ChunkActor) const;
+#endif
+
+private:
+	/** Cached values for computing contact offsets */
+	float ContactOffsetFactor; 
+	float MinContactOffset;
+	float MaxContactOffset;
+
+	friend struct FApexDestructionSyncActors;
+};
+
+
+
diff --git a/Engine/Plugins/Runtime/ApexDestruction/Source/ApexDestruction/Public/DestructibleFractureSettings.h b/Engine/Plugins/Runtime/ApexDestruction/Source/ApexDestruction/Public/DestructibleFractureSettings.h
new file mode 100644
--- /dev/null
+++ b/Engine/Plugins/Runtime/ApexDestruction/Source/ApexDestruction/Public/DestructibleFractureSettings.h
@@ -0,0 +1,209 @@
+// Copyright Epic Games, Inc. All Rights Reserved.
+
+#pragma once
+
+#include "CoreMinimal.h"
+#include "UObject/ObjectMacros.h"
+#include "UObject/Object.h"
+#include "EngineDefines.h"
+#include "DestructibleFractureSettings.generated.h"
+
+class UMaterialInterface;
+
+#if WITH_APEX
+
+// Forward declares
+namespace nvidia
+{
+	namespace apex
+	{
+		struct FractureMaterialDesc;
+		class DestructibleChunkDesc;
+		class DestructibleGeometryDesc;
+		class DestructibleAsset;
+		class DestructibleAssetAuthoring;
+		struct ExplicitRenderTriangle;
+		struct ExplicitSubmeshData;
+		class DestructibleAssetCookingDesc;
+	};
+};
+#endif // WITH_APEX
+
+
+/** 
+ * Options for APEX asset import.
+ **/
+namespace EDestructibleImportOptions
+{
+	enum Type
+	{
+		// Just imports the APEX asset
+		None				= 0,
+		// Preserves settings in DestructibleMesh 
+		PreserveSettings	= 1<<0,
+	};
+};
+
+
+/** Parameters to describe the application of U,V coordinates on a particular slice within a destructible. */
+struct UE_DEPRECATED(4.26, "APEX is deprecated. Destruction in future will be supported using Chaos Destruction.") FFractureMaterial;
+USTRUCT()
+struct FFractureMaterial
+{
+	GENERATED_USTRUCT_BODY()
+
+	/**
+	 * The UV scale (geometric distance/unit texture distance) for interior materials.
+	 * Default = (100.0f,100.0f).
+	 */
+	UPROPERTY(EditAnywhere, Category=FractureMaterial)
+	FVector2D	UVScale;
+
+	/**
+	 * A UV origin offset for interior materials.
+	 * Default = (0.0f,0.0f).
+	 */
+	UPROPERTY(EditAnywhere, Category=FractureMaterial)
+	FVector2D	UVOffset;
+
+	/**
+	 * Object-space vector specifying surface tangent direction.  If this vector is (0.0f,0.0f,0.0f), then an arbitrary direction will be chosen.
+	 * Default = (0.0f,0.0f,0.0f).
+	 */
+	UPROPERTY(EditAnywhere, Category=FractureMaterial)
+	FVector		Tangent;
+
+	/**
+	 * Angle from tangent direction for the u coordinate axis.
+	 * Default = 0.0f.
+	 */
+	UPROPERTY(EditAnywhere, Category=FractureMaterial)
+	float		UAngle;
+
+	/**
+	 * The element index to use for the newly-created triangles.
+	 * If a negative index is given, a new element will be created for interior triangles.
+	 * Default = -1
+	 */
+	UPROPERTY(EditAnywhere, Category=FractureMaterial)
+	int32		InteriorElementIndex;
+
+	FFractureMaterial()
+		: UVScale(100.0f, 100.0f)
+		, UVOffset(0.0f, 0.0f)
+		, Tangent(0.0f, 0.0f, 0.0f)
+		, UAngle(0.0f)
+		, InteriorElementIndex(-1)
+	{}
+
+	//~ Begin FFractureMaterial Interface
+#if WITH_APEX
+	void	FillNxFractureMaterialDesc(nvidia::apex::FractureMaterialDesc& PFractureMaterialDesc);
+#endif // WITH_APEX
+	//~ End FFractureMaterial Interface
+};
+
+
+/** Per-chunk authoring data. */
+USTRUCT()
+struct UE_DEPRECATED(4.26, "APEX is deprecated. Destruction in future will be supported using Chaos Destruction.") FDestructibleChunkParameters
+{
+	GENERATED_USTRUCT_BODY()
+
+	/**
+		Defines the chunk to be environmentally supported, if the appropriate NxDestructibleParametersFlag flags
+		are set in NxDestructibleParameters.
+		Default = false.
+	*/
+	UPROPERTY(EditAnywhere, Category=DestructibleChunkParameters)
+	bool bIsSupportChunk;
+
+	/**
+		Defines the chunk to be unfractureable.  If this is true, then none of its children will be fractureable.
+		Default = false.
+	*/
+	UPROPERTY(EditAnywhere, Category=DestructibleChunkParameters)
+	bool bDoNotFracture;
+
+	/**
+		Defines the chunk to be undamageable.  This means this chunk will not fracture, but its children might.
+		Default = false.
+	*/
+	UPROPERTY(EditAnywhere, Category=DestructibleChunkParameters)
+	bool bDoNotDamage;
+
+	/**
+		Defines the chunk to be uncrumbleable.  This means this chunk will not be broken down into fluid mesh particles
+		no matter how much damage it takes.  Note: this only applies to chunks with no children.  For a chunk with
+		children, then:
+		1) The chunk may be broken down into its children, and then its children may be crumbled, if the doNotCrumble flag
+		is not set on them.
+		2) If the Destructible module's chunk depth offset LOD may be set such that this chunk effectively has no children.
+		In this case, the doNotCrumble flag will apply to it.
+		Default = false.
+	*/
+	UPROPERTY(EditAnywhere, Category=DestructibleChunkParameters)
+	bool bDoNotCrumble;
+
+	FDestructibleChunkParameters()
+		: bIsSupportChunk(false)
+		, bDoNotFracture(false)
+		, bDoNotDamage(false)
+		, bDoNotCrumble(false)
+	{}
+};
+
+
+/** Information to create an NxDestructibleAsset */
+UCLASS(MinimalAPI)
+class UE_DEPRECATED(4.26, "APEX is deprecated. Destruction in future will be supported using Chaos Destruction.") UDestructibleFractureSettings
+	: public UObject
+{
+	GENERATED_UCLASS_BODY()
+
+	/**  Destructor. Defined to avoid compilation errors on clang where the UHT generated CPP file 
+	 * 	ends up instantiating destructors for member variables without full type information
+	 */
+	virtual ~UDestructibleFractureSettings();
+
+	/** DO NOT USE. This constructor is for internal usage only for hot-reload purposes. */
+	UDestructibleFractureSettings(FVTableHelper& Helper);
+
+	/** The number of voronoi cell sites. */
+	UPROPERTY(EditAnywhere, Category = Voronoi, meta = (ClampMin = "1", UIMin = "1"))
+	int32									CellSiteCount;
+
+PRAGMA_DISABLE_DEPRECATION_WARNINGS
+	/** Stored interior material data.  Just need one as we only support Voronoi splitting. */
+	UPROPERTY(EditAnywhere, transient, Category=General)
+	FFractureMaterial						FractureMaterialDesc;
+PRAGMA_ENABLE_DEPRECATION_WARNINGS
+
+	/** Random seed for reproducibility */
+	UPROPERTY(EditAnywhere, Category=General)
+	int32									RandomSeed;
+
+	/** Stored Voronoi sites */
+	UPROPERTY()
+	TArray<FVector>							VoronoiSites;
+
+	/** The mesh's original number of submeshes.  APEX needs to store this in the authoring. */
+	UPROPERTY()
+	int32									OriginalSubmeshCount;
+
+	/** APEX references materials by name, but we'll bypass that mechanism and use of UE materials instead. */
+	UPROPERTY()
+	TArray<UMaterialInterface*>				Materials;
+
+PRAGMA_DISABLE_DEPRECATION_WARNINGS
+	/** Per-chunk authoring parameters, which should be made writable when a chunk selection GUI is in place. */
+	UPROPERTY()
+	TArray<FDestructibleChunkParameters>	ChunkParameters;
+PRAGMA_ENABLE_DEPRECATION_WARNINGS	
+
+	//~ Begin UObject Interface.
+	virtual void PostInitProperties() override;
+	virtual void BeginDestroy() override;
+	virtual void Serialize(FArchive& Ar) override;
+	//~ End  UObject Interface
+};
diff --git a/Engine/Plugins/Runtime/ApexDestruction/Source/ApexDestruction/Public/DestructibleMesh.h b/Engine/Plugins/Runtime/ApexDestruction/Source/ApexDestruction/Public/DestructibleMesh.h
new file mode 100644
--- /dev/null
+++ b/Engine/Plugins/Runtime/ApexDestruction/Source/ApexDestruction/Public/DestructibleMesh.h
@@ -0,0 +1,571 @@
+// Copyright Epic Games, Inc. All Rights Reserved.
+
+#pragma once
+
+#include "CoreMinimal.h"
+#include "UObject/ObjectMacros.h"
+#include "EngineDefines.h"
+#include "Engine/SkeletalMesh.h"
+#include "Misc/CoreMiscDefines.h"
+#include "DestructibleMesh.generated.h"
+
+class UPhysicalMaterial;
+class UStaticMesh;
+
+#if WITH_APEX
+namespace NvParameterized
+{
+	class Interface;
+}
+
+namespace nvidia
+{
+	namespace apex
+	{
+		class DestructibleAsset;
+	}
+}
+#endif
+
+#ifndef WITH_DESTRUCTIBLE_DEPRECATION
+#define WITH_DESTRUCTIBLE_DEPRECATION 0
+#endif
+
+/**
+	Chunks up to the depth DefaultImpactDamageDepth will take impact damage, unless IDO_On or IDO_Off is chosen.
+*/
+UENUM()
+enum EImpactDamageOverride
+{
+	IDO_None,
+	IDO_On,
+	IDO_Off,
+	IDO_MAX,
+};
+
+
+/** Properties that may be set for all chunks at a particular depth in the fracture hierarchy. */
+USTRUCT()
+struct UE_DEPRECATED(4.26, "APEX is deprecated. Destruction in future will be supported using Chaos Destruction.") FDestructibleDepthParameters
+{
+	GENERATED_USTRUCT_BODY()
+
+	/** Chunks up to the depth DefaultImpactDamageDepth will take impact damage, unless one of the override options (see EImpactDamageOverride) is chosen. */
+	UPROPERTY(EditAnywhere, Category=DestructibleDepthParameters)
+	TEnumAsByte<enum EImpactDamageOverride> ImpactDamageOverride;
+
+
+	FDestructibleDepthParameters()
+		: ImpactDamageOverride(IDO_None)
+	{}
+
+#if WITH_APEX
+	void FillDestructibleActorDesc(NvParameterized::Interface* Params, const char* OverrideName, const char* OverrideValueName) const;
+	void LoadDefaultDestructibleParametersFromApexAsset(const NvParameterized::Interface* Params, const char* OverrideName, const char* OverrideValueName);
+#endif
+
+};
+
+
+/** Flags that apply to a destructible actor. */
+USTRUCT()
+struct UE_DEPRECATED(4.26, "APEX is deprecated. Destruction in future will be supported using Chaos Destruction.") FDestructibleParametersFlag
+{
+	GENERATED_USTRUCT_BODY()
+
+	/**
+		If set, chunks will "remember" damage applied to them, so that many applications of a damage amount
+		below damageThreshold will eventually fracture the chunk.  If not set, a single application of
+		damage must exceed damageThreshold in order to fracture the chunk.
+	*/
+	UPROPERTY(EditAnywhere, Category=DestructibleParametersFlag)
+	uint32 bAccumulateDamage:1;
+
+	/**
+		If set, then chunks which are tagged as "support" chunks (via NxDestructibleChunkDesc::isSupportChunk)
+		will have environmental support in static destructibles.
+
+		Note: if both bAssetDefinedSupport and bWorldSupport are set, then chunks must be tagged as
+		"support" chunks AND overlap the Scene's static geometry in order to be environmentally supported.
+	*/
+	UPROPERTY(EditAnywhere, Category=DestructibleParametersFlag)
+	uint32 bAssetDefinedSupport:1;
+
+	/**
+		If set, then chunks which overlap the Scene's static geometry will have environmental support in
+		static destructibles.
+
+		Note: if both bAssetDefinedSupport and bWorldSupport are set, then chunks must be tagged as
+		"support" chunks AND overlap the Scene's static geometry in order to be environmentally supported.
+	*/
+	UPROPERTY(EditAnywhere, Category=DestructibleParametersFlag)
+	uint32 bWorldSupport:1;
+
+	/**
+		Whether or not chunks at or deeper than the "debris" depth (see NxDestructibleParameters::debrisDepth)
+		will time out.  The lifetime is a value between NxDestructibleParameters::debrisLifetimeMin and
+		NxDestructibleParameters::debrisLifetimeMax, based upon the destructible module's LOD setting.
+	*/
+	UPROPERTY(EditAnywhere, Category=DestructibleParametersFlag)
+	uint32 bDebrisTimeout:1;
+
+	/**
+		Whether or not chunks at or deeper than the "debris" depth (see NxDestructibleParameters::debrisDepth)
+		will be removed if they separate too far from their origins.  The maxSeparation is a value between
+		NxDestructibleParameters::debrisMaxSeparationMin and NxDestructibleParameters::debrisMaxSeparationMax,
+		based upon the destructible module's LOD setting.
+	*/
+	UPROPERTY(EditAnywhere, Category=DestructibleParametersFlag)
+	uint32 bDebrisMaxSeparation:1;
+
+	/**
+		If set, the smallest chunks may be further broken down, either by fluid crumbles (if a crumble particle
+		system is specified in the NxDestructibleActorDesc), or by simply removing the chunk if no crumble
+		particle system is specified.  Note: the "smallest chunks" are normally defined to be the deepest level
+		of the fracture hierarchy.  However, they may be taken from higher levels of the hierarchy if
+		NxModuleDestructible::setMaxChunkDepthOffset is called with a non-zero value.
+	*/
+	UPROPERTY(EditAnywhere, Category=DestructibleParametersFlag)
+	uint32 bCrumbleSmallestChunks:1;
+
+	/**
+		If set, the NxDestructibleActor::rayCast function will search within the nearest visible chunk hit
+		for collisions with child chunks.  This is used to get a better raycast position and normal, in
+		case the parent collision volume does not tightly fit the graphics mesh.  The returned chunk index
+		will always be that of the visible parent that is intersected, however.
+	*/
+	UPROPERTY(EditAnywhere, Category=DestructibleParametersFlag)
+	uint32 bAccurateRaycasts:1;
+
+	/**
+		If set, the ValidBounds field of NxDestructibleParameters will be used.  These bounds are translated
+		(but not scaled or rotated) to the origin of the destructible actor.  If a chunk or chunk island moves
+		outside of those bounds, it is destroyed.
+	*/
+	UPROPERTY(EditAnywhere, Category=DestructibleParametersFlag)
+	uint32 bUseValidBounds:1;
+
+	/**
+		If initially static, the destructible will become part of an extended support structure if it is
+		in contact with another static destructible that also has this flag set.
+	*/
+	UPROPERTY(EditAnywhere, Category=DestructibleParametersFlag)
+	uint32 bFormExtendedStructures:1;
+
+
+	FDestructibleParametersFlag()
+		: bAccumulateDamage(false)
+		, bAssetDefinedSupport(false)
+		, bWorldSupport(false)
+		, bDebrisTimeout(false)
+		, bDebrisMaxSeparation(false)
+		, bCrumbleSmallestChunks(false)
+		, bAccurateRaycasts(false)
+		, bUseValidBounds(false)
+		, bFormExtendedStructures(false)
+	{}
+
+#if WITH_APEX
+	void FillDestructibleActorDesc(NvParameterized::Interface* Params) const;
+	void LoadDefaultDestructibleParametersFromApexAsset(const NvParameterized::Interface* Params);
+#endif
+
+};
+
+
+/** Parameters that pertain to chunk damage. */
+USTRUCT()
+struct UE_DEPRECATED(4.26, "APEX is deprecated. Destruction in future will be supported using Chaos Destruction.") FDestructibleDamageParameters
+{
+	GENERATED_USTRUCT_BODY()
+
+	/**
+		The damage amount which will cause a chunk to fracture (break free).
+	*/
+	UPROPERTY(EditAnywhere, Category=DestructibleDamageParameters)
+	float DamageThreshold;
+
+	/** Controls how easily damage spreads. DamageRadius = Damage*DamageSpread. All chunks within DamageRadius will take damage. Full damage is taken at zero distance, and zero damage at the DamageRadius.*/
+	UPROPERTY(EditAnywhere, Category=DestructibleDamageParameters)
+	float DamageSpread;
+
+	/**
+		Whether to apply damage to destructible when colliding with an object.
+		@see ImpactDamage
+	*/
+	UPROPERTY(EditAnywhere, Category = DestructibleDamageParameters)
+	bool bEnableImpactDamage;
+
+	/**
+		Controls how much damage is applied upon collision. Damage = ImpactDamage * ImpactForce.
+		@see DepthParameters for per level control of ImpactDamage
+	*/
+	UPROPERTY(EditAnywhere, Category=DestructibleDamageParameters, meta = (editcondition = "bEnableImpactDamage") )
+	float ImpactDamage;
+
+	/**
+		Max depth level where impact damage is enabled. @see DepthParameters for per level control of ImpactDamage
+		If negative, impact damage is disabled
+	*/
+	UPROPERTY(EditAnywhere, Category = DestructibleDamageParameters, meta = (ClampMin = "0", UIMin = "0", editcondition = "bEnableImpactDamage"))
+	int32 DefaultImpactDamageDepth;
+
+	/**
+		By default, objects that collide with destructibles will bounce back. Custom resistance allows for finer control of how much a destructible "pushes back" against a colliding object.
+		@see ImpactResistance
+	*/
+	UPROPERTY(EditAnywhere, Category = DestructibleDamageParameters)
+	bool bCustomImpactResistance;
+
+	/**
+		Controls how much resistance is applied to colliding objects. Weak materials like glass should set this to a low value so that objects will pass right through them during fracture.
+		@see DepthParameters for per level control of ImpactResistance
+	*/
+	UPROPERTY(EditAnywhere, meta = (ClampMin = "0.001", UIMin = "0", editcondition = "bCustomImpactResistance"), Category = DestructibleDamageParameters)
+	float ImpactResistance;
+
+
+	FDestructibleDamageParameters()
+		: DamageThreshold(1.0f)
+		, DamageSpread(0.1f)
+		, bEnableImpactDamage(false)
+		, ImpactDamage(0.1f)
+		, DefaultImpactDamageDepth(0.f)
+		, bCustomImpactResistance(false)
+		, ImpactResistance(1.f)
+	{}
+
+#if WITH_APEX
+	void FillDestructibleActorDesc(NvParameterized::Interface* Params, UPhysicalMaterial* PhysMat) const;
+	void LoadDefaultDestructibleParametersFromApexAsset(const NvParameterized::Interface* Params);
+#endif
+	
+};
+
+
+/** Parameters that pertain to chunk debris-level settings. */
+USTRUCT()
+struct UE_DEPRECATED(4.26, "APEX is deprecated. Destruction in future will be supported using Chaos Destruction.") FDestructibleDebrisParameters
+{
+	GENERATED_USTRUCT_BODY()
+
+	/**
+		"Debris chunks" (see debrisDepth, above) will be destroyed after a time (in seconds)
+		separated from non-debris chunks.  The actual lifetime is interpolated between these
+		two bDebrisTimeout, based upon the module's LOD setting.  To disable lifetime, clear the
+		bDebrisTimeout flag in the flags field.
+		If debrisLifetimeMax < debrisLifetimeMin, the mean of the two is used for both.
+		Default debrisLifetimeMin = 1.0, debrisLifetimeMax = 10.0f.
+	*/
+	UPROPERTY(EditAnywhere, Category=DestructibleDebrisParameters)
+	float DebrisLifetimeMin;
+
+	UPROPERTY(EditAnywhere, Category=DestructibleDebrisParameters)
+	float DebrisLifetimeMax;
+
+	/**
+		"Debris chunks" (see debrisDepth, above) will be destroyed if they are separated from
+		their origin by a distance greater than maxSeparation.  The actual maxSeparation is
+		interpolated between these two values, based upon the module's LOD setting.  To disable
+		maxSeparation, clear the bDebrisMaxSeparation flag in the flags field.
+		If debrisMaxSeparationMax < debrisMaxSeparationMin, the mean of the two is used for both.
+		Default debrisMaxSeparationMin = 1.0, debrisMaxSeparationMax = 10.0f.
+	*/
+	UPROPERTY(EditAnywhere, Category=DestructibleDebrisParameters)
+	float DebrisMaxSeparationMin;
+
+	UPROPERTY(EditAnywhere, Category=DestructibleDebrisParameters)
+	float DebrisMaxSeparationMax;
+
+	/**
+		"Debris chunks" (see debrisDepth, above) will be destroyed if they leave this box.
+		The box translates with the destructible actor's initial position, but does not
+		rotate or scale.
+	*/
+	UPROPERTY(EditAnywhere, Category=DestructibleDebrisParameters)
+	FBox ValidBounds;
+
+
+
+	FDestructibleDebrisParameters()
+		: DebrisLifetimeMin(1.0f)
+		, DebrisLifetimeMax(10.0f)
+		, DebrisMaxSeparationMin(1.0f)
+		, DebrisMaxSeparationMax(10.0f)
+		, ValidBounds(FVector(-500000.0f), FVector(500000.0f))
+	{}
+
+#if WITH_APEX
+	void FillDestructibleActorDesc(NvParameterized::Interface* Params) const;
+	void LoadDefaultDestructibleParametersFromApexAsset(const NvParameterized::Interface* Params);
+#endif
+	
+};
+
+
+/** Parameters that are less-often used. */
+USTRUCT()
+struct UE_DEPRECATED(4.26, "APEX is deprecated. Destruction in future will be supported using Chaos Destruction.") FDestructibleAdvancedParameters
+{
+	GENERATED_USTRUCT_BODY()
+
+	/**
+		Limits the amount of damage applied to a chunk.  This is useful for preventing the entire destructible
+		from getting pulverized by a very large application of damage.  This can easily happen when impact damage is
+		used, and the damage amount is proportional to the impact force (see forceToDamage).
+	*/
+	UPROPERTY(EditAnywhere, Category=DestructibleAdvancedParameters)
+	float DamageCap;
+
+	/** 
+		Large impact force may be reported if rigid bodies are spawned inside one another.  In this case the realative velocity of the two
+		objects will be low.  This variable allows the user to set a minimum velocity threshold for impacts to ensure that the objects are 
+		moving at a min velocity in order for the impact force to be considered.  
+	*/
+	UPROPERTY(EditAnywhere, Category=DestructibleAdvancedParameters)
+	float ImpactVelocityThreshold;
+
+	/**
+		If greater than 0, the chunks' speeds will not be allowed to exceed this value.  Use 0
+		to disable this feature (this is the default).
+	*/
+	UPROPERTY(EditAnywhere, Category=DestructibleAdvancedParameters)
+	float MaxChunkSpeed;
+
+	/**
+		Scale factor used to apply an impulse force along the normal of chunk when fractured.  This is used
+		in order to "push" the pieces out as they fracture.
+	*/
+	UPROPERTY(EditAnywhere, Category=DestructibleAdvancedParameters)
+	float FractureImpulseScale;
+
+
+	FDestructibleAdvancedParameters()
+		: DamageCap(0)
+		, ImpactVelocityThreshold(0)
+		, MaxChunkSpeed(0)
+		, FractureImpulseScale(0)
+	{}
+
+#if WITH_APEX
+	void FillDestructibleActorDesc(NvParameterized::Interface* Params) const;
+	void LoadDefaultDestructibleParametersFromApexAsset(const NvParameterized::Interface* Params);
+#endif
+
+};
+
+
+/** Special hierarchy depths for various behaviors. */
+USTRUCT()
+struct UE_DEPRECATED(4.26, "APEX is deprecated. Destruction in future will be supported using Chaos Destruction.") FDestructibleSpecialHierarchyDepths
+{
+	GENERATED_USTRUCT_BODY()
+
+	// todo dynamicChunksDominanceGroup
+	// todo dynamicChunksGroupsMask
+	/**
+		The chunk hierarchy depth at which to create a support graph.  Higher depth levels give more detailed support,
+		but will give a higher computational load.  Chunks below the support depth will never be supported.
+	*/
+	UPROPERTY(EditAnywhere, Category=DestructibleSpecialHierarchyDepths)
+	int32 SupportDepth;
+
+	/**
+		The chunks will not be broken free below this depth.
+	*/
+	UPROPERTY(EditAnywhere, Category=DestructibleSpecialHierarchyDepths)
+	int32 MinimumFractureDepth;
+
+
+	/** 
+		Enables debris at a specific depth level.
+		@see DebrisDepth */
+	UPROPERTY(EditAnywhere, Category = DestructibleSpecialHierarchyDepths)
+	bool bEnableDebris;
+
+	/**
+		The hierarchy depth at which chunks are considered to be "debris."  
+	*/
+	UPROPERTY(EditAnywhere, Category = DestructibleSpecialHierarchyDepths, meta = (UIMin = "0", editcondition = "bEnableDebris"))
+	int32 DebrisDepth;
+
+	/**
+		The chunk hierarchy depth up to which chunks will always be processed.  These chunks are considered
+		to be essential either for gameplay or visually.
+		The minimum value is 0, meaning the level 0 chunk is always considered essential.
+		Default value is 0.
+	*/
+	UPROPERTY(EditAnywhere, Category=DestructibleSpecialHierarchyDepths, meta=(DisplayName = "Essential LOD Depth"))
+	int32 EssentialDepth;
+
+	/**
+	 * Default constructor.
+	 */
+	FDestructibleSpecialHierarchyDepths()
+		: SupportDepth(0)
+		, MinimumFractureDepth(0)
+		, bEnableDebris(false)
+		, DebrisDepth(-1)
+		, EssentialDepth(0)
+	{}
+
+#if WITH_APEX
+	void FillDestructibleActorDesc(NvParameterized::Interface* Params) const;
+	void LoadDefaultDestructibleParametersFromApexAsset(const NvParameterized::Interface* Params);
+#endif
+};
+
+
+/** Parameters that apply to a destructible actor. */
+USTRUCT()
+struct UE_DEPRECATED(4.26, "APEX is deprecated. Destruction in future will be supported using Chaos Destruction.") FDestructibleParameters
+{
+	GENERATED_USTRUCT_BODY()
+
+	/** Parameters that pertain to chunk damage.  See DestructibleDamageParameters. */
+	UPROPERTY(EditAnywhere, Category=DestructibleParameters)
+	struct FDestructibleDamageParameters DamageParameters;
+
+	/** Parameters that pertain to chunk debris-level settings.  See DestructibleDebrisParameters. */
+	UPROPERTY(EditAnywhere, Category=DestructibleParameters)
+	struct FDestructibleDebrisParameters DebrisParameters;
+
+	/** Parameters that are less-often used.  See DestructibleAdvancedParameters. */
+	UPROPERTY(EditAnywhere, Category=DestructibleParameters)
+	struct FDestructibleAdvancedParameters AdvancedParameters;
+
+	/** Special hierarchy depths for various behaviors. */
+	UPROPERTY(EditAnywhere, Category=DestructibleParameters)
+	struct FDestructibleSpecialHierarchyDepths SpecialHierarchyDepths;
+
+	/**
+		Parameters that apply to every chunk at a given level.
+		the element [0] of the array applies to the level 0 (unfractured) chunk, element [1] applies
+		to the level 1 chunks, etc.
+	*/
+	UPROPERTY(EditAnywhere, editfixedsize, Category=DestructibleParameters)
+	TArray<struct FDestructibleDepthParameters> DepthParameters;
+
+	/** A collection of flags defined in DestructibleParametersFlag. */
+	UPROPERTY(EditAnywhere, Category=DestructibleParameters)
+	struct FDestructibleParametersFlag Flags;
+};
+
+
+/**
+ * Holds an APEX destructible asset as well as an associated USkeletalMesh.
+ */
+UCLASS(hidecategories = (Object, Mesh, LevelOfDetail, Mirroring, Physics, Reimport, Clothing), MinimalAPI)
+class UE_DEPRECATED(4.26, "APEX is deprecated. Destruction in future will be supported using Chaos Destruction.") UDestructibleMesh : public USkeletalMesh
+{
+	GENERATED_UCLASS_BODY()
+
+	/** Parameters controlling the destruction behavior. */
+	UPROPERTY(EditAnywhere, Category=DestructibleMesh)
+	struct FDestructibleParameters DefaultDestructibleParameters;
+
+	/** Fracture effects for each fracture level, unless overridden in the component. */
+	UPROPERTY(EditAnywhere, editfixedsize, Category=Effects)
+	TArray<struct FFractureEffect> FractureEffects;
+
+#if WITH_EDITORONLY_DATA
+	/** Information used to author an NxDestructibleAsset*/
+	UPROPERTY(instanced)
+	class UDestructibleFractureSettings* FractureSettings;
+	
+	/** Static mesh this destructible mesh is created from. Is nullptr if not created from a static mesh */
+	UPROPERTY()
+	UStaticMesh* SourceStaticMesh;
+	
+	/** Timestamp of the source static meshes last import at the time this destruction mesh has been generated. */
+	UPROPERTY()
+	FDateTime SourceSMImportTimestamp;
+
+	/** Array of static meshes to build the fracture chunks from */
+	UPROPERTY(Transient)
+	TArray<UStaticMesh*> FractureChunkMeshes;
+#endif // WITH_EDITORONLY_DATA
+
+public:
+
+#if WITH_APEX
+	/** ApexDestructibleAsset is a pointer to the Apex asset interface for this destructible asset */
+	nvidia::apex::DestructibleAsset* ApexDestructibleAsset;
+#endif // WITH_APEX
+
+public:
+
+	// UObject interface.
+
+	virtual void PostLoad() override;
+	virtual void Serialize(FArchive& Ar) override;
+	virtual void FinishDestroy() override;
+#if WITH_EDITOR
+	virtual void PreEditChange(FProperty* PropertyAboutToChange) override;
+
+	virtual bool HasCustomActorFactory() const
+	{
+		return true;
+	}
+
+	virtual bool HasCustomActorReimportFactory() const
+	{
+		return true;
+	}
+#endif
+
+public:
+		
+#if WITH_APEX
+	/** 
+	 * Retrieve a default actor descriptor.
+	 * Builds the descriptor from the NxDestructibleAsset and the overrides provided in DefaultDestructibleParameters
+	 * @return : The (NvParameterized) descriptor.
+	 */
+	NvParameterized::Interface* GetDestructibleActorDesc(UPhysicalMaterial* PhysMat);
+
+	/**
+	 * Access to APEX native destructible asset .
+	 * @return : Returns the NxDestructibleAsset stored in this object.
+	 */
+	nvidia::apex::DestructibleAsset*		GetApexDestructibleAsset()
+	{
+		return ApexDestructibleAsset;
+	}
+#endif // WITH_APEX
+
+	/** Fills DefaultDestructibleParameters with parameters from the NxDestructibleAsset. */
+	APEXDESTRUCTION_API void LoadDefaultDestructibleParametersFromApexAsset();
+
+	/**
+	 * Create DestructibleFractureSettings for this DestructibleMesh if it doesn't have one.
+	 */
+	APEXDESTRUCTION_API void CreateFractureSettings();
+
+	/**
+	 * Imports FractureSettings data from a StaticMesh
+	 *
+	 * @param StaticMesh The StaticMesh to import
+	 * @return true on success, false otherwise.
+	 */
+	APEXDESTRUCTION_API bool BuildFractureSettingsFromStaticMesh(UStaticMesh* StaticMesh);
+
+	/** 
+	 * Initializes this DestructibleMesh from a StaticMesh.
+	 *
+	 * @param StaticMesh The StaticMesh to import
+	 * @return true on success, false otherwise.
+	 */
+	APEXDESTRUCTION_API bool BuildFromStaticMesh(UStaticMesh& StaticMesh);
+
+	/** 
+	 * Initialized this DestructibleMesh from the StaticMesh it was created from and the passed
+	 * in ChunkMeshes to build the level 1 chunks from.
+	 *
+	 * @param ChunkMeshes Meshes to build the level1 chunks from
+	 * @return true on success, false otherwise.
+	 */
+	APEXDESTRUCTION_API bool SetupChunksFromStaticMeshes(const TArray<UStaticMesh*>& ChunkMeshes);
+};
+
+APEXDESTRUCTION_API DECLARE_LOG_CATEGORY_EXTERN(LogDestructible, Warning, All);
\ No newline at end of file
diff --git a/Engine/Plugins/Runtime/ApexDestruction/Source/ApexDestructionEditor/ApexDestructionEditor.Build.cs b/Engine/Plugins/Runtime/ApexDestruction/Source/ApexDestructionEditor/ApexDestructionEditor.Build.cs
new file mode 100644
--- /dev/null
+++ b/Engine/Plugins/Runtime/ApexDestruction/Source/ApexDestructionEditor/ApexDestructionEditor.Build.cs
@@ -0,0 +1,41 @@
+// Copyright Epic Games, Inc. All Rights Reserved.
+
+using UnrealBuildTool;
+using System.IO;
+
+public class ApexDestructionEditor : ModuleRules
+{
+	public ApexDestructionEditor(ReadOnlyTargetRules Target) : base(Target)
+	{
+		PrivateDependencyModuleNames.AddRange(
+			new string[] {
+				"Core",
+				"CoreUObject", // @todo Mac: for some reason it's needed to link in debug on Mac
+				"Engine",
+                "ApexDestruction",
+                "AssetTools",
+//                "PhysX",
+//                "APEX",
+                "InputCore",
+                "RenderCore",
+                "RHI",
+                "Slate",
+                "SlateCore",
+                "EditorStyle",
+				"EditorFramework",
+                "UnrealEd",
+                "DesktopPlatform",
+                "AssetRegistry",
+                "ContentBrowser",
+                "Projects",
+//                "ApexDestructionLib"
+            }
+		);
+
+        AddEngineThirdPartyPrivateStaticDependencies(Target,
+            "FBX"
+        );
+
+        SetupModulePhysicsSupport(Target);
+}
+}
diff --git a/Engine/Plugins/Runtime/ApexDestruction/Source/ApexDestructionEditor/Private/ActorFactoryDestructible.cpp b/Engine/Plugins/Runtime/ApexDestruction/Source/ApexDestructionEditor/Private/ActorFactoryDestructible.cpp
new file mode 100644
--- /dev/null
+++ b/Engine/Plugins/Runtime/ApexDestruction/Source/ApexDestructionEditor/Private/ActorFactoryDestructible.cpp
@@ -0,0 +1,353 @@
+// Copyright Epic Games, Inc. All Rights Reserved.
+
+#include "ActorFactoryDestructible.h"
+#include "DestructibleActor.h"
+#include "DestructibleMesh.h"
+#include "DestructibleComponent.h"
+//#include "DestructibleMeshFactory.h"
+#include "Misc/FileHelper.h"
+#include "ActorFactories/ActorFactory.h"
+#include "AssetRegistry/AssetData.h"
+#include "PhysicsPublic.h"
+#include "PhysXIncludes.h"
+#include "Editor.h"
+#include "Misc/MessageDialog.h"
+#include "EditorFramework/AssetImportData.h"
+//#include "ApexDestructibleAssetImport.h"
+//#include "ApexClothingUtils.h"
+//#include "ReimportDestructibleMeshFactory.h"
+
+#define LOCTEXT_NAMESPACE "ActorFactoryDestructible"
+
+DEFINE_LOG_CATEGORY_STATIC(LogDestructibleFactories, Log, All);
+
+/*-----------------------------------------------------------------------------
+UActorFactoryDestructible
+-----------------------------------------------------------------------------*/
+UActorFactoryDestructible::UActorFactoryDestructible(const FObjectInitializer& ObjectInitializer)
+	: Super(ObjectInitializer)
+{
+PRAGMA_DISABLE_DEPRECATION_WARNINGS
+	NewActorClass = ADestructibleActor::StaticClass();
+PRAGMA_ENABLE_DEPRECATION_WARNINGS
+
+	DisplayName = LOCTEXT("DestructibleDisplayName", "Destructible");
+	bUseSurfaceOrientation = true;
+}
+
+bool UActorFactoryDestructible::CanCreateActorFrom(const FAssetData& AssetData, FText& OutErrorMsg)
+{
+PRAGMA_DISABLE_DEPRECATION_WARNINGS
+	if (!AssetData.IsValid() || !AssetData.GetClass()->IsChildOf(UDestructibleMesh::StaticClass()))
+	{
+		OutErrorMsg = NSLOCTEXT("CanCreateActor", "NoDestructibleMeshSpecified", "No destructible mesh was specified.");
+		return false;
+	}
+PRAGMA_ENABLE_DEPRECATION_WARNINGS
+	return true;
+}
+
+void UActorFactoryDestructible::PostSpawnActor(UObject* Asset, AActor* NewActor)
+{
+	Super::PostSpawnActor(Asset, NewActor);
+
+PRAGMA_DISABLE_DEPRECATION_WARNINGS
+UDestructibleMesh* DestructibleMesh = CastChecked<UDestructibleMesh>(Asset);
+	ADestructibleActor* NewDestructibleActor = CastChecked<ADestructibleActor>(NewActor);
+
+	// Term Component
+	NewDestructibleActor->GetDestructibleComponent()->UnregisterComponent();
+
+	// Change properties
+	NewDestructibleActor->GetDestructibleComponent()->SetSkeletalMesh(DestructibleMesh);
+
+	// Init Component
+	NewDestructibleActor->GetDestructibleComponent()->RegisterComponent();
+PRAGMA_ENABLE_DEPRECATION_WARNINGS
+}
+
+UObject* UActorFactoryDestructible::GetAssetFromActorInstance(AActor* Instance)
+{
+PRAGMA_DISABLE_DEPRECATION_WARNINGS
+	check(Instance->IsA(NewActorClass));
+	ADestructibleActor* DA = CastChecked<ADestructibleActor>(Instance);
+
+	check(DA->GetDestructibleComponent());
+	return DA->GetDestructibleComponent()->SkeletalMesh;
+PRAGMA_ENABLE_DEPRECATION_WARNINGS
+}
+
+void UActorFactoryDestructible::PostCreateBlueprint(UObject* Asset, AActor* CDO)
+{
+PRAGMA_DISABLE_DEPRECATION_WARNINGS
+	if (Asset != NULL && CDO != NULL)
+	{
+		UDestructibleMesh* DestructibleMesh = CastChecked<UDestructibleMesh>(Asset);
+		ADestructibleActor* DestructibleActor = CastChecked<ADestructibleActor>(CDO);
+
+		DestructibleActor->GetDestructibleComponent()->SetSkeletalMesh(DestructibleMesh);
+	}
+PRAGMA_ENABLE_DEPRECATION_WARNINGS
+}
+
+FQuat UActorFactoryDestructible::AlignObjectToSurfaceNormal(const FVector& InSurfaceNormal, const FQuat& ActorRotation) const
+{
+	// Meshes align the Z (up) axis with the surface normal
+	return FindActorAlignmentRotation(ActorRotation, FVector(0.f, 0.f, 1.f), InSurfaceNormal);
+}
+
+/*------------------------------------------------------------------------------
+UDestructibleMeshFactory implementation.
+------------------------------------------------------------------------------*/
+// namespace DestructibleFactoryConstants
+// {
+// 	static const FString DestructibleAssetClass("DestructibleAssetParameters");
+// }
+// 
+// UDestructibleMeshFactory::UDestructibleMeshFactory(const FObjectInitializer& ObjectInitializer)
+// 	: Super(ObjectInitializer)
+// {
+// 
+// 	bEditorImport = true;
+// 
+// PRAGMA_DISABLE_DEPRECATION_WARNINGS
+// 	SupportedClass = UDestructibleMesh::StaticClass();
+// PRAGMA_ENABLE_DEPRECATION_WARNINGS
+// 
+// 	bCreateNew = false;
+// 	Formats.Add(TEXT("apx;APEX XML Asset"));
+// 	Formats.Add(TEXT("apb;APEX Binary Asset"));
+// }
+// 
+// FText UDestructibleMeshFactory::GetDisplayName() const
+// {
+// 	return LOCTEXT("APEXDestructibleFactoryDescription", "APEX Destructible Asset");
+// }
+// 
+// #if WITH_APEX
+// 
+// bool UDestructibleMeshFactory::FactoryCanImport(const FString& Filename)
+// {
+// 	// Need to read in the file and try to create an asset to get it's type
+// 	TArray<uint8> FileBuffer;
+// 	if (FFileHelper::LoadFileToArray(FileBuffer, *Filename, FILEREAD_Silent))
+// 	{
+// 		physx::PxFileBuf* Stream = GApexSDK->createMemoryReadStream(FileBuffer.GetData(), FileBuffer.Num());
+// 		if (Stream)
+// 		{
+// 			NvParameterized::Serializer::SerializeType SerializeType = GApexSDK->getSerializeType(*Stream);
+// 			if (NvParameterized::Serializer* Serializer = GApexSDK->createSerializer(SerializeType))
+// 			{
+// 				NvParameterized::Serializer::DeserializedData DeserializedData;
+// 				Serializer->deserialize(*Stream, DeserializedData);
+// 
+// 				if (DeserializedData.size() > 0)
+// 				{
+// 					NvParameterized::Interface* AssetInterface = DeserializedData[0];
+// 
+// 					int32 StringLength = StringCast<TCHAR>(AssetInterface->className()).Length();
+// 					FString ClassName(StringLength, StringCast<TCHAR>(AssetInterface->className()).Get());
+// 
+// 					if (ClassName == DestructibleFactoryConstants::DestructibleAssetClass)
+// 					{
+// 						return true;
+// 					}
+// 				}
+// 			}
+// 
+// 			GApexSDK->releaseMemoryReadStream(*Stream);
+// 		}
+// 	}
+// 	return false;
+// }
+// 
+// UObject* UDestructibleMeshFactory::FactoryCreateBinary
+// (
+// 	UClass*				Class,
+// 	UObject*			InParent,
+// 	FName				Name,
+// 	EObjectFlags		Flags,
+// 	UObject*			Context,
+// 	const TCHAR*		FileType,
+// 	const uint8*&		Buffer,
+// 	const uint8*			BufferEnd,
+// 	FFeedbackContext*	Warn
+// )
+// {
+// 	GEditor->GetEditorSubsystem<UImportSubsystem>()->BroadcastAssetPreImport(this, Class, InParent, Name, FileType);
+// 
+// PRAGMA_DISABLE_DEPRECATION_WARNINGS
+// 	// The return value
+// 	UDestructibleMesh* DestructibleMesh = nullptr;
+// 
+// 	// Create an Apex NxDestructibleAsset from the binary blob
+// 	apex::DestructibleAsset* ApexDestructibleAsset = CreateApexDestructibleAssetFromBuffer(Buffer, (int32)(BufferEnd - Buffer));
+// 	if (ApexDestructibleAsset != nullptr)
+// 	{
+// 		// Succesfully created the NxDestructibleAsset, now create a UDestructibleMesh
+// 		DestructibleMesh = ImportDestructibleMeshFromApexDestructibleAsset(InParent, *ApexDestructibleAsset, Name, Flags, nullptr);
+// 		if (DestructibleMesh != nullptr)
+// 		{
+// 			GEditor->GetEditorSubsystem<UImportSubsystem>()->BroadcastAssetPostImport(this, DestructibleMesh);
+// 
+// 			// Success
+// 			DestructibleMesh->PostEditChange();
+// 		}
+// 	}
+// #if WITH_APEX_CLOTHING
+// 	else
+// 	{
+// 		// verify whether this is an Apex Clothing asset or not 
+// 		apex::ClothingAsset* ApexClothingAsset = ApexClothingUtils::CreateApexClothingAssetFromBuffer(Buffer, (int32)(BufferEnd - Buffer));
+// 
+// 		if (ApexClothingAsset)
+// 		{
+// 			FMessageDialog::Open(EAppMsgType::Ok, LOCTEXT("ApexClothingWrongImport", "The file you tried to import is an APEX clothing asset file. You need to use Persona to import this asset and associate it with a skeletal mesh.\n\n 1. Import a skeletal mesh from an FBX file, or choose an existing skeletal asset and open it up in Persona.\n 2. Choose \"Add APEX clothing file\" and choose this APEX clothing asset file."));
+// 
+// 			// This asset is used only for showing a message how to import an Apex Clothing asset properly
+// 			GPhysCommandHandler->DeferredRelease(ApexClothingAsset);
+// 		}
+// 	}
+// #endif // #if WITH_APEX_CLOTHING
+// 
+// 	return DestructibleMesh;
+// PRAGMA_ENABLE_DEPRECATION_WARNINGS
+// }
+// 
+// #endif // WITH_APEX
+// 
+// /*-----------------------------------------------------------------------------
+// UReimportDestructibleMeshFactory implementation.
+// -----------------------------------------------------------------------------*/
+// UReimportDestructibleMeshFactory::UReimportDestructibleMeshFactory(const FObjectInitializer& ObjectInitializer)
+// 	: Super(ObjectInitializer)
+// {
+// PRAGMA_DISABLE_DEPRECATION_WARNINGS
+// 	SupportedClass = UDestructibleMesh::StaticClass();
+// PRAGMA_ENABLE_DEPRECATION_WARNINGS
+// 
+// 	bCreateNew = false;
+// 	bText = false;
+// 	Formats.Add(TEXT("apx;APEX XML Asset"));
+// 	Formats.Add(TEXT("apb;APEX Binary Asset"));
+// 
+// }
+// 
+// FText UReimportDestructibleMeshFactory::GetDisplayName() const
+// {
+// 	return LOCTEXT("APEXReimportDestructibleAssetFactoryDescription", "APEX Reimport Destructible Asset");
+// }
+// 
+// #if WITH_APEX
+// 
+// bool UReimportDestructibleMeshFactory::CanReimport(UObject* Obj, TArray<FString>& OutFilenames)
+// {
+// PRAGMA_DISABLE_DEPRECATION_WARNINGS
+// 	UDestructibleMesh* DestructibleMesh = Cast<UDestructibleMesh>(Obj);
+// 	if (DestructibleMesh)
+// 	{
+// 		if (DestructibleMesh->AssetImportData)
+// 		{
+// 			DestructibleMesh->AssetImportData->ExtractFilenames(OutFilenames);
+// 		}
+// 		else
+// 		{
+// 			OutFilenames.Add(TEXT(""));
+// 		}
+// 		return true;
+// 	}
+// 	return false;
+// PRAGMA_ENABLE_DEPRECATION_WARNINGS
+// }
+// 
+// void UReimportDestructibleMeshFactory::SetReimportPaths(UObject* Obj, const TArray<FString>& NewReimportPaths)
+// {
+// PRAGMA_DISABLE_DEPRECATION_WARNINGS
+// 	UDestructibleMesh* DestructibleMesh = Cast<UDestructibleMesh>(Obj);
+// 	if (DestructibleMesh && ensure(NewReimportPaths.Num() == 1))
+// 	{
+// 		DestructibleMesh->AssetImportData->UpdateFilenameOnly(NewReimportPaths[0]);
+// 	}
+// PRAGMA_ENABLE_DEPRECATION_WARNINGS
+// }
+// 
+// EReimportResult::Type UReimportDestructibleMeshFactory::Reimport(UObject* Obj)
+// {
+// PRAGMA_DISABLE_DEPRECATION_WARNINGS
+// 	// Only handle valid skeletal meshes
+// 	if (!Obj || !Obj->IsA(UDestructibleMesh::StaticClass()))
+// 	{
+// 		return EReimportResult::Failed;
+// 	}
+// 
+// 	UDestructibleMesh* DestructibleMesh = Cast<UDestructibleMesh>(Obj);
+// 
+// 	const FString Filename = DestructibleMesh->AssetImportData->GetFirstFilename();
+// 
+// 	// If there is no file path provided, can't reimport from source
+// 	if (!Filename.Len())
+// 	{
+// 		// Since this is a new system most skeletal meshes don't have paths, so logging has been commented out
+// 		//UE_LOG(LogEditorFactories, Warning, TEXT("-- cannot reimport: skeletal mesh resource does not have path stored."));
+// 		return EReimportResult::Failed;
+// 	}
+// 
+// 	UE_LOG(LogDestructibleFactories, Log, TEXT("Performing atomic reimport of [%s]"), *Filename);
+// 
+// 	// Ensure that the file provided by the path exists
+// 	if (IFileManager::Get().FileSize(*Filename) == INDEX_NONE)
+// 	{
+// 		UE_LOG(LogDestructibleFactories, Warning, TEXT("-- cannot reimport: source file cannot be found."));
+// 		return EReimportResult::Failed;
+// 	}
+// 
+// 	CurrentFilename = Filename;
+// 
+// 	// Create an Apex NxDestructibleAsset from the binary blob
+// 	apex::DestructibleAsset* ApexDestructibleAsset = CreateApexDestructibleAssetFromFile(Filename);
+// 	if (ApexDestructibleAsset != nullptr)
+// 	{
+// 		// Succesfully created the NxDestructibleAsset, now create a UDestructibleMesh
+// 		UDestructibleMesh* ReimportedDestructibleMesh = ImportDestructibleMeshFromApexDestructibleAsset(DestructibleMesh->GetOuter(), *ApexDestructibleAsset, DestructibleMesh->GetFName(), DestructibleMesh->GetFlags(), nullptr,
+// 			EDestructibleImportOptions::PreserveSettings);
+// 		if (ReimportedDestructibleMesh != nullptr)
+// 		{
+// 			check(ReimportedDestructibleMesh == DestructibleMesh);
+// 
+// 			UE_LOG(LogDestructibleFactories, Log, TEXT("-- imported successfully"));
+// 
+// 			// Try to find the outer package so we can dirty it up
+// 			if (DestructibleMesh->GetOuter())
+// 			{
+// 				DestructibleMesh->GetOuter()->MarkPackageDirty();
+// 			}
+// 			else
+// 			{
+// 				DestructibleMesh->MarkPackageDirty();
+// 			}
+// 		}
+// 		else
+// 		{
+// 			FMessageDialog::Open(EAppMsgType::Ok, NSLOCTEXT("UnrealEd", "ImportFailed_Destructible", "Reimport Failed"));
+// 			UE_LOG(LogDestructibleFactories, Warning, TEXT("-- import failed"));
+// 		}
+// 	}
+// 	else
+// 	{
+// 		FMessageDialog::Open(EAppMsgType::Ok, NSLOCTEXT("UnrealEd", "ImportFailed_Destructible", "Reimport Failed"));
+// 		UE_LOG(LogDestructibleFactories, Warning, TEXT("-- import failed"));
+// 	}
+// 
+// PRAGMA_ENABLE_DEPRECATION_WARNINGS
+// 
+// 	return EReimportResult::Succeeded;
+// }
+// 
+// int32 UReimportDestructibleMeshFactory::GetPriority() const
+// {
+// 	return ImportPriority;
+// }
+
+//#endif // #if WITH_APEX
+
+#undef LOCTEXT_NAMESPACE
diff --git a/Engine/Plugins/Runtime/ApexDestruction/Source/ApexDestructionEditor/Private/ActorFactoryDestructible.h b/Engine/Plugins/Runtime/ApexDestruction/Source/ApexDestructionEditor/Private/ActorFactoryDestructible.h
new file mode 100644
--- /dev/null
+++ b/Engine/Plugins/Runtime/ApexDestruction/Source/ApexDestructionEditor/Private/ActorFactoryDestructible.h
@@ -0,0 +1,29 @@
+// Copyright Epic Games, Inc. All Rights Reserved.
+
+
+#pragma once
+
+#include "CoreMinimal.h"
+#include "UObject/ObjectMacros.h"
+#include "ActorFactories/ActorFactory.h"
+#include "ActorFactoryDestructible.generated.h"
+
+class AActor;
+struct FAssetData;
+
+UCLASS(MinimalAPI, config=Editor)
+class UActorFactoryDestructible : public UActorFactory
+{
+	GENERATED_UCLASS_BODY()
+
+	//~ Begin UActorFactory Interface
+	virtual void PostSpawnActor( UObject* Asset, AActor* NewActor ) override;
+	virtual void PostCreateBlueprint( UObject* Asset, AActor* CDO ) override;
+	virtual bool CanCreateActorFrom( const FAssetData& AssetData, FText& OutErrorMsg ) override;
+	virtual UObject* GetAssetFromActorInstance(AActor* ActorInstance) override;
+	virtual FQuat AlignObjectToSurfaceNormal(const FVector& InSurfaceNormal, const FQuat& ActorRotation) const override;
+	//~ End UActorFactory Interface
+};
+
+
+
diff --git a/Engine/Plugins/Runtime/ApexDestruction/Source/ApexDestructionEditor/Private/ApexDestructionEditorModule.cpp b/Engine/Plugins/Runtime/ApexDestruction/Source/ApexDestructionEditor/Private/ApexDestructionEditorModule.cpp
new file mode 100644
--- /dev/null
+++ b/Engine/Plugins/Runtime/ApexDestruction/Source/ApexDestructionEditor/Private/ApexDestructionEditorModule.cpp
@@ -0,0 +1,271 @@
+// Copyright Epic Games, Inc. All Rights Reserved.
+
+#include "ApexDestructionEditorModule.h"
+//#include "Misc/PackageName.h"
+//#include "Modules/ModuleManager.h"
+//#include "IDestructibleMeshEditor.h"
+//#include "DestructibleMeshEditor.h"
+//#include "Misc/MessageDialog.h"
+//
+#include "IAssetTools.h"
+#include "AssetToolsModule.h"
+//
+#include "DestructibleMesh.h"
+//#include "Engine/StaticMesh.h"
+//#include "Materials/Material.h"
+//#include "Materials/MaterialExpressionSpeedTree.h"
+//
+//#include "AssetToolsModule.h"
+#include "AssetTypeActions_DestructibleMesh.h"
+//#include "ContentBrowserDelegates.h"
+//#include "Framework/MultiBox/MultiBoxBuilder.h"
+//#include "ContentBrowserModule.h"
+//
+#include "DestructibleMeshThumbnailRenderer.h"
+#include "ThumbnailRendering/ThumbnailManager.h"
+#include "DestructibleMeshComponentBroker.h"
+//#include "ApexDestructionStyle.h"
+//#include "PropertyEditorModule.h"
+//#include "DestructibleMeshDetails.h"
+
+PRAGMA_DISABLE_DEPRECATION_WARNINGS
+IMPLEMENT_MODULE( FDestructibleMeshEditorModule, ApexDestructionEditor );
+PRAGMA_ENABLE_DEPRECATION_WARNINGS
+
+#define LOCTEXT_NAMESPACE "DestructibleMeshEditor"
+
+//const FName DestructibleMeshEditorAppIdentifier = FName(TEXT("DestructibleMeshEditorApp"));
+
+
+void FDestructibleMeshEditorModule::StartupModule()
+{
+//	MenuExtensibilityManager = MakeShareable(new FExtensibilityManager);
+//	ToolBarExtensibilityManager = MakeShareable(new FExtensibilityManager);
+//
+	IAssetTools& AssetTools = FModuleManager::GetModuleChecked<FAssetToolsModule>("AssetTools").Get();
+//
+	AssetAction = MakeShareable(new FAssetTypeActions_DestructibleMesh);
+	AssetTools.RegisterAssetTypeActions(AssetAction.ToSharedRef());
+PRAGMA_DISABLE_DEPRECATION_WARNINGS
+	UThumbnailManager::Get().RegisterCustomRenderer(UDestructibleMesh::StaticClass(), UDestructibleMeshThumbnailRenderer::StaticClass());
+
+	DestructibleMeshComponentBroker = MakeShareable(new FDestructibleMeshComponentBroker);
+	FComponentAssetBrokerage::RegisterBroker(DestructibleMeshComponentBroker, UDestructibleComponent::StaticClass(), false, true);
+PRAGMA_ENABLE_DEPRECATION_WARNINGS
+//
+//	if (!IsRunningCommandlet())
+//	{
+//		FContentBrowserModule& ContentBrowserModule = FModuleManager::LoadModuleChecked<FContentBrowserModule>(TEXT("ContentBrowser"));
+//		TArray<FContentBrowserMenuExtender_SelectedAssets>& CBMenuExtenderDelegates = ContentBrowserModule.GetAllAssetViewContextMenuExtenders();
+//
+//		CBMenuExtenderDelegates.Add(FContentBrowserMenuExtender_SelectedAssets::CreateStatic(&FDestructibleMeshEditorModule::OnExtendContentBrowserAssetSelectionMenu));
+//		ContentBrowserExtenderDelegateHandle = CBMenuExtenderDelegates.Last().GetHandle();
+//	}
+//
+//PRAGMA_DISABLE_DEPRECATION_WARNINGS
+//	UThumbnailManager::Get().RegisterCustomRenderer(UDestructibleMesh::StaticClass(), UDestructibleMeshThumbnailRenderer::StaticClass());
+//
+//	DestructibleMeshComponentBroker = MakeShareable(new FDestructibleMeshComponentBroker);
+//	FComponentAssetBrokerage::RegisterBroker(DestructibleMeshComponentBroker, UDestructibleComponent::StaticClass(), false, true);
+//
+//	FPropertyEditorModule& PropertyModule = FModuleManager::LoadModuleChecked<FPropertyEditorModule>("PropertyEditor");
+//	PropertyModule.RegisterCustomClassLayout(UDestructibleMesh::StaticClass()->GetFName(), FOnGetDetailCustomizationInstance::CreateStatic(&FDestructibleMeshDetails::MakeInstance));
+//
+//	FApexDestructionStyle::Initialize();
+//PRAGMA_ENABLE_DEPRECATION_WARNINGS
+}
+
+//TSharedRef<FExtender> FDestructibleMeshEditorModule::OnExtendContentBrowserAssetSelectionMenu(const TArray<FAssetData>& SelectedAssets)
+//{
+//
+//	TSharedRef<FExtender> Extender(new FExtender());
+//
+//#if WITH_APEX
+//	// Run thru the assets to determine if any meet our criteria
+//	bool bAnyStaticMesh = false;
+//	for (auto AssetIt = SelectedAssets.CreateConstIterator(); AssetIt; ++AssetIt)
+//	{
+//		const FAssetData& Asset = *AssetIt;
+//		bAnyStaticMesh = bAnyStaticMesh || (Asset.AssetClass == UStaticMesh::StaticClass()->GetFName());
+//	}
+//
+//	if (bAnyStaticMesh)
+//	{
+//		// Add the sprite actions sub-menu extender
+//		Extender->AddMenuExtension("GetAssetActions", EExtensionHook::After, nullptr, FMenuExtensionDelegate::CreateLambda([SelectedAssets](FMenuBuilder& MenuBuilder)
+//		{
+//			MenuBuilder.AddMenuEntry(
+//				NSLOCTEXT("AssetTypeActions_StaticMesh", "ObjectContext_CreateDestructibleMesh", "Create Destructible Mesh"),
+//				NSLOCTEXT("AssetTypeActions_StaticMesh", "ObjectContext_CreateDestructibleMeshTooltip", "Creates a DestructibleMesh from the StaticMesh and opens it in the DestructibleMesh editor."),
+//				FSlateIcon(FEditorStyle::GetStyleSetName(), "ClassIcon.DestructibleComponent"),
+//				FUIAction(FExecuteAction::CreateStatic(&FAssetTypeActions_DestructibleMesh::ExecuteCreateDestructibleMeshes, SelectedAssets), FCanExecuteAction()));
+//		}));
+//	}
+//#endif
+//
+//	return Extender;
+//}
+
+
+void FDestructibleMeshEditorModule::ShutdownModule()
+{
+PRAGMA_DISABLE_DEPRECATION_WARNINGS
+//	MenuExtensibilityManager.Reset();
+//	ToolBarExtensibilityManager.Reset();
+//
+	if (FModuleManager::Get().IsModuleLoaded("AssetTools"))
+	{
+		IAssetTools& AssetTools = FModuleManager::GetModuleChecked<FAssetToolsModule>("AssetTools").Get();
+		AssetTools.UnregisterAssetTypeActions(AssetAction.ToSharedRef());
+	}
+//
+//	FApexDestructionStyle::Shutdown();
+//
+//	if(FModuleManager::Get().IsModuleLoaded("ContentBrowser"))
+//	{
+//		FContentBrowserModule& ContentBrowserModule = FModuleManager::LoadModuleChecked<FContentBrowserModule>(TEXT("ContentBrowser"));
+//		TArray<FContentBrowserMenuExtender_SelectedAssets>& CBMenuExtenderDelegates = ContentBrowserModule.GetAllAssetViewContextMenuExtenders();
+//		CBMenuExtenderDelegates.RemoveAll([this](const FContentBrowserMenuExtender_SelectedAssets& Delegate) { return Delegate.GetHandle() == ContentBrowserExtenderDelegateHandle; });
+//	}
+//	
+//
+	if (UObjectInitialized())
+	{
+		UThumbnailManager::Get().UnregisterCustomRenderer(UDestructibleMesh::StaticClass());
+		FComponentAssetBrokerage::UnregisterBroker(DestructibleMeshComponentBroker);
+	}
+PRAGMA_ENABLE_DEPRECATION_WARNINGS
+}
+
+
+//PRAGMA_DISABLE_DEPRECATION_WARNINGS
+//TSharedRef<IDestructibleMeshEditor> FDestructibleMeshEditorModule::CreateDestructibleMeshEditor( const EToolkitMode::Type Mode, const TSharedPtr< IToolkitHost >& InitToolkitHost, UDestructibleMesh* Table )
+//{
+//	TSharedRef< FDestructibleMeshEditor > NewDestructibleMeshEditor( new FDestructibleMeshEditor() );
+//	NewDestructibleMeshEditor->InitDestructibleMeshEditor( Mode, InitToolkitHost, Table );
+//	return NewDestructibleMeshEditor;
+//}
+//PRAGMA_ENABLE_DEPRECATION_WARNINGS
+//
+//PRAGMA_DISABLE_DEPRECATION_WARNINGS
+//UDestructibleMesh* FDestructibleMeshEditorModule::CreateDestructibleMeshFromStaticMesh(UObject* InParent, UStaticMesh* StaticMesh, FName Name, EObjectFlags Flags, FText& OutErrorMsg)
+//{
+//#if !WITH_APEX
+//	OutErrorMsg = LOCTEXT("NoAPEX", "APEX is not built for the current engine configuration.");
+//	return nullptr;
+//#endif
+//
+//	if (StaticMesh == NULL)
+//	{
+//		OutErrorMsg = LOCTEXT( "StaticMeshInvalid", "Static Mesh is Invalid!" );
+//		return NULL;
+//	}
+//
+//	// We can't use any speedtree materials for destructibles due to UV requirements, so check the graph here
+//	TArray<UMaterial*> SpeedTreeMaterials;
+//	for(FStaticMaterial& StaticMaterial : StaticMesh->GetStaticMaterials())
+//	{
+//		if(UMaterialInterface* Mat = StaticMaterial.MaterialInterface)
+//		{
+//			if(UMaterial* BaseMat = Mat->GetBaseMaterial())
+//			{
+//				for(UMaterialExpression* Expression : BaseMat->Expressions)
+//				{
+//					if(Cast<UMaterialExpressionSpeedTree>(Expression))
+//					{
+//						SpeedTreeMaterials.Add(BaseMat);
+//						break;
+//					}
+//				}
+//			}
+//		}
+//	}
+//	
+//	if(SpeedTreeMaterials.Num() > 0)
+//	{
+//		// Invalid, can't create a mesh from this
+//		FTextBuilder TextBuilder;
+//		TextBuilder.AppendLine(FText::Format(LOCTEXT("StaticMeshInvalid_SpeedTree", "The static mesh '{0}' uses SpeedTree materials which are not compatible with destructible meshes. Cannot create destructible.\n\nList of Materials:\n"), FText::FromString(StaticMesh->GetName())));
+//	
+//		for(UMaterial* SpeedTreeMat : SpeedTreeMaterials)
+//		{
+//			TextBuilder.AppendLine(SpeedTreeMat->GetName());
+//		}
+//	
+//		FMessageDialog::Open(EAppMsgType::Ok, TextBuilder.ToText());
+//		return nullptr;
+//	}
+//
+//
+//	FString DestructibleName;
+//	
+//	if (Name == NAME_None)
+//	{
+//		DestructibleName = StaticMesh->GetName();
+//		DestructibleName += TEXT("_DM");
+//	}
+//	else
+//	{
+//		DestructibleName = *Name.ToString();
+//	}
+//
+//	UDestructibleMesh* DestructibleMesh = FindObject<UDestructibleMesh>( InParent, *DestructibleName );
+//
+//	// If there is an existing mesh, ask the user if they want to use the existing mesh, replace the existing mesh, or create a new mesh.
+//	if (DestructibleMesh != NULL)
+//	{
+//		enum ENameConflictDialogReturnType
+//		{
+//			NCD_None = 0,
+//			NCD_Use,
+//			NCD_Replace,
+//			NCD_Create
+//		};
+//#if 0	// BRG - removed SSimpleDialog, needs replacing
+//		FSimpleDialogModule& SimpleDialogModule = FModuleManager::LoadModuleChecked<FSimpleDialogModule>( TEXT("SimpleDialog") );
+//		SimpleDialogModule.CreateSimpleDialog(NSLOCTEXT("SimpleDialogModule", "DestructibleMeshAlreadyExistsTitle", "DestructibleMesh Already Exists").ToString(), NSLOCTEXT("SimpleDialogModule", "DestructibleMeshAlreadyExistsMessage", "A DestructibleMesh by the same name already exists.  Select an action.").ToString());
+//		SimpleDialogModule.AddButton(NCD_Use, NSLOCTEXT("SimpleDialogModule", "UseExistingDestructibleMesh", "Use").ToString(), NSLOCTEXT("SimpleDialogModule", "UseExistingDestructibleMeshTip", "Open the DestructibleMesh editor with the existing mesh").ToString());
+//		SimpleDialogModule.AddButton(NCD_Replace, NSLOCTEXT("SimpleDialogModule", "ReplaceExistingDestructibleMesh", "Replace").ToString(), NSLOCTEXT("SimpleDialogModule", "ReplaceExistingDestructibleMeshTip", "Replace the existing DestructibleMesh with a new one, and show in the DestructibleMesh editor").ToString());
+//		SimpleDialogModule.AddButton(NCD_Create, NSLOCTEXT("SimpleDialogModule", "CreateNewDestructibleMesh", "Create").ToString(), NSLOCTEXT("SimpleDialogModule", "CreateNewDestructibleMeshTip", "Create a new DestructibleMesh with a different name, and show in the DestructibleMesh editor").ToString());
+//		const uint32 UserResponse = SimpleDialogModule.ShowSimpleDialog();
+//#endif
+//		const uint32 UserResponse = NCD_None;
+//		CA_SUPPRESS(6326);
+//		switch (UserResponse)
+//		{
+//		default:	// Default to using the same mesh
+//		case NCD_Use:
+//			return DestructibleMesh;
+//		case NCD_Replace:
+//			break;
+//		case NCD_Create:
+//			Name = MakeUniqueObjectName(StaticMesh->GetOuter(), UDestructibleMesh::StaticClass(), Name);
+//			DestructibleName = *Name.ToString();
+//			DestructibleMesh = NULL;	// So that one will be created
+//			break;
+//		}
+//	}
+//	
+//	// Create the new UDestructibleMesh object if we still haven't found one after a possible resolution above
+//	if (DestructibleMesh == NULL)
+//	{
+//		FAssetToolsModule& AssetToolsModule = FModuleManager::LoadModuleChecked<FAssetToolsModule>("AssetTools");
+//		DestructibleMesh = Cast<UDestructibleMesh>(AssetToolsModule.Get().CreateAsset(DestructibleName, FPackageName::GetLongPackagePath(InParent->GetOutermost()->GetPathName()), UDestructibleMesh::StaticClass(), NULL));
+//
+//		if(!DestructibleMesh)
+//		{
+//			FFormatNamedArguments Arguments;
+//			Arguments.Add(TEXT("Name"), FText::FromString( StaticMesh->GetName() ));
+//
+//			OutErrorMsg = FText::Format( LOCTEXT( "DestructibleMeshFailed", "Failed to Create Destructible Mesh Asset from {Name}!" ), Arguments );
+//			return NULL;
+//		}
+//	}
+//
+//	DestructibleMesh->BuildFromStaticMesh(*StaticMesh);
+//	
+//	return DestructibleMesh;
+//}
+//PRAGMA_ENABLE_DEPRECATION_WARNINGS
+
+#undef LOCTEXT_NAMESPACE
diff --git a/Engine/Plugins/Runtime/ApexDestruction/Source/ApexDestructionEditor/Private/ApexDestructionEditorModule.h b/Engine/Plugins/Runtime/ApexDestruction/Source/ApexDestructionEditor/Private/ApexDestructionEditorModule.h
new file mode 100644
--- /dev/null
+++ b/Engine/Plugins/Runtime/ApexDestruction/Source/ApexDestructionEditor/Private/ApexDestructionEditorModule.h
@@ -0,0 +1,73 @@
+// Copyright Epic Games, Inc. All Rights Reserved.
+
+#pragma once
+
+#include "CoreMinimal.h"
+#include "Modules/ModuleInterface.h"
+//#include "Toolkits/IToolkitHost.h"
+//#include "Toolkits/AssetEditorToolkit.h"
+//#include "AssetData.h"
+
+//class UStaticMesh;
+//class IDestructibleMeshEditor;
+//class FExtender;
+
+class IAssetTypeActions;
+
+
+//extern const FName DestructibleMeshEditorAppIdentifier;
+
+//class IDestructibleMeshEditor;
+
+/** DestructibleMesh Editor module */
+
+class UE_DEPRECATED(4.26, "APEX is deprecated. Destruction in future will be supported using Chaos Destruction.") FDestructibleMeshEditorModule : public IModuleInterface
+{
+
+public:
+	// IModuleInterface
+	virtual void StartupModule() override;
+	virtual void ShutdownModule() override;
+
+//	PRAGMA_DISABLE_DEPRECATION_WARNINGS
+//	/**
+//	 * Creates an instance of table editor object.  Only virtual so that it can be called across the DLL boundary.
+//	 *
+//	 * @param	Mode					Mode that this editor should operate in
+//	 * @param	InitToolkitHost			When Mode is WorldCentric, this is the level editor instance to spawn this editor within
+//	 * @param	Table					The table to start editing
+//	 *
+//	 * @return	Interface to the new table editor
+//	 */
+//	virtual TSharedRef<IDestructibleMeshEditor> CreateDestructibleMeshEditor( const EToolkitMode::Type Mode, const TSharedPtr< class IToolkitHost >& InitToolkitHost, class UDestructibleMesh* Table );
+//
+//	/**
+//	 * Creates a DestructibleMesh from a StaticMesh
+//	 *
+//	 * @param InParent - parent object for the UDestructibleMesh
+//	 * @param StaticMesh - the StaticMesh to convert
+//	 * @param Name - the Unreal name for the UDestructibleMesh.  If Name == NAME_None, then a the StaticMesh name followed appended with "_Dest" will be used.
+//	 * @param Flags - object flags for the UDestructibleMesh
+//	 * @param OutErrorMsg - if it returns false, this string contains the error
+//	 *
+//	 * @return The newly created UDestructibleMesh if successful, NULL otherwise
+//	 */
+//	virtual class UDestructibleMesh*	CreateDestructibleMeshFromStaticMesh(UObject* InParent, UStaticMesh* StaticMesh, FName Name, EObjectFlags Flags, FText& OutErrorMsg);
+//	PRAGMA_ENABLE_DEPRECATION_WARNINGS
+//
+//	/** Gets the extensibility managers for outside entities to extend static mesh editor's menus and toolbars */
+//	virtual TSharedPtr<FExtensibilityManager> GetMenuExtensibilityManager() override {return MenuExtensibilityManager;}
+//	virtual TSharedPtr<FExtensibilityManager> GetToolBarExtensibilityManager() override {return ToolBarExtensibilityManager;}
+//
+//private:
+//	TSharedPtr<FExtensibilityManager> MenuExtensibilityManager;
+//	TSharedPtr<FExtensibilityManager> ToolBarExtensibilityManager;
+	TSharedPtr<class IAssetTypeActions> AssetAction;
+	TSharedPtr<class FDestructibleMeshComponentBroker> DestructibleMeshComponentBroker;
+//	FDelegateHandle ContentBrowserExtenderDelegateHandle;
+//
+//private:
+//	static TSharedRef<FExtender> OnExtendContentBrowserAssetSelectionMenu(const TArray<FAssetData>& SelectedAssets);
+};
+
+
diff --git a/Engine/Plugins/Runtime/ApexDestruction/Source/ApexDestructionEditor/Private/AssetTypeActions_DestructibleMesh.cpp b/Engine/Plugins/Runtime/ApexDestruction/Source/ApexDestructionEditor/Private/AssetTypeActions_DestructibleMesh.cpp
new file mode 100644
--- /dev/null
+++ b/Engine/Plugins/Runtime/ApexDestruction/Source/ApexDestructionEditor/Private/AssetTypeActions_DestructibleMesh.cpp
@@ -0,0 +1,36 @@
+// Copyright Epic Games, Inc. All Rights Reserved.
+
+#include "AssetTypeActions_DestructibleMesh.h"
+#include "DestructibleMesh.h"
+//#include "ApexDestructionEditorModule.h"
+#include "Engine/StaticMesh.h"
+#include "Widgets/Notifications/SNotificationList.h"
+#include "Framework/Notifications/NotificationManager.h"
+#include "Subsystems/AssetEditorSubsystem.h"
+
+#define LOCTEXT_NAMESPACE "AssetTypeActions"
+
+UClass* FAssetTypeActions_DestructibleMesh::GetSupportedClass() const
+{
+PRAGMA_DISABLE_DEPRECATION_WARNINGS
+	return UDestructibleMesh::StaticClass();
+PRAGMA_ENABLE_DEPRECATION_WARNINGS
+}
+
+void FAssetTypeActions_DestructibleMesh::OpenAssetEditor( const TArray<UObject*>& InObjects, TSharedPtr<IToolkitHost> EditWithinLevelEditor )
+{
+//	for (auto ObjIt = InObjects.CreateConstIterator(); ObjIt; ++ObjIt)
+//	{
+//PRAGMA_DISABLE_DEPRECATION_WARNINGS
+//		auto Mesh = Cast<UDestructibleMesh>(*ObjIt);
+//
+//		if (Mesh != NULL)
+//		{
+//			FDestructibleMeshEditorModule& DestructibleMeshEditorModule = FModuleManager::LoadModuleChecked<FDestructibleMeshEditorModule>( "ApexDestructionEditor" );
+//			TSharedRef< IDestructibleMeshEditor > NewDestructibleMeshEditor = DestructibleMeshEditorModule.CreateDestructibleMeshEditor( EToolkitMode::Standalone, EditWithinLevelEditor, Mesh );
+//		}
+//PRAGMA_ENABLE_DEPRECATION_WARNINGS
+//	}
+}
+
+#undef LOCTEXT_NAMESPACE
diff --git a/Engine/Plugins/Runtime/ApexDestruction/Source/ApexDestructionEditor/Private/AssetTypeActions_DestructibleMesh.h b/Engine/Plugins/Runtime/ApexDestruction/Source/ApexDestructionEditor/Private/AssetTypeActions_DestructibleMesh.h
new file mode 100644
--- /dev/null
+++ b/Engine/Plugins/Runtime/ApexDestruction/Source/ApexDestructionEditor/Private/AssetTypeActions_DestructibleMesh.h
@@ -0,0 +1,20 @@
+// Copyright Epic Games, Inc. All Rights Reserved.
+
+#pragma once
+
+#include "CoreMinimal.h"
+#include "Toolkits/IToolkitHost.h"
+#include "AssetTypeActions_Base.h"
+
+class FAssetTypeActions_DestructibleMesh : public FAssetTypeActions_Base
+{
+public:
+	// IAssetTypeActions Implementation
+	virtual FText GetName() const override { return NSLOCTEXT("AssetTypeActions", "AssetTypeActions_DestructibleMesh", "Destructible Mesh"); }
+	virtual FColor GetTypeColor() const override { return FColor(200,128,128); }
+	virtual UClass* GetSupportedClass() const override;
+	virtual uint32 GetCategories() override { return EAssetTypeCategories::Physics; }
+	virtual bool HasActions( const TArray<UObject*>& InObjects ) const override { return false; }
+	virtual void OpenAssetEditor( const TArray<UObject*>& InObjects, TSharedPtr<class IToolkitHost> EditWithinLevelEditor = TSharedPtr<IToolkitHost>() ) override;
+	virtual bool IsImportedAsset() const override { return false; }
+};
diff --git a/Engine/Plugins/Runtime/ApexDestruction/Source/ApexDestructionEditor/Private/DestructibleMeshComponentBroker.h b/Engine/Plugins/Runtime/ApexDestruction/Source/ApexDestructionEditor/Private/DestructibleMeshComponentBroker.h
new file mode 100644
--- /dev/null
+++ b/Engine/Plugins/Runtime/ApexDestruction/Source/ApexDestructionEditor/Private/DestructibleMeshComponentBroker.h
@@ -0,0 +1,45 @@
+// Copyright Epic Games, Inc. All Rights Reserved.
+
+#pragma once
+
+#include "ComponentAssetBroker.h"
+#include "DestructibleMesh.h"
+#include "DestructibleComponent.h"
+
+//////////////////////////////////////////////////////////////////////////
+// FDestructibleMeshComponentBroker
+PRAGMA_DISABLE_DEPRECATION_WARNINGS
+class FDestructibleMeshComponentBroker : public IComponentAssetBroker
+{
+public:
+	UClass* GetSupportedAssetClass() override
+	{
+		return UDestructibleMesh::StaticClass();
+	}
+
+	virtual bool AssignAssetToComponent(UActorComponent* InComponent, UObject* InAsset) override
+	{
+		if (UDestructibleComponent* DestMeshComp = Cast<UDestructibleComponent>(InComponent))
+		{
+			UDestructibleMesh* DMesh = Cast<UDestructibleMesh>(InAsset);
+
+			if (DMesh)
+			{
+				DestMeshComp->SetDestructibleMesh(DMesh);
+				return true;
+			}
+		}
+
+		return false;
+	}
+
+	virtual UObject* GetAssetFromComponent(UActorComponent* InComponent) override
+	{
+		if (UDestructibleComponent* DestMeshComp = Cast<UDestructibleComponent>(InComponent))
+		{
+			return DestMeshComp->GetDestructibleMesh();
+		}
+		return NULL;
+	}
+};
+PRAGMA_ENABLE_DEPRECATION_WARNINGS
diff --git a/Engine/Plugins/Runtime/ApexDestruction/Source/ApexDestructionEditor/Private/DestructibleMeshThumbnailRenderer.cpp b/Engine/Plugins/Runtime/ApexDestruction/Source/ApexDestructionEditor/Private/DestructibleMeshThumbnailRenderer.cpp
new file mode 100644
--- /dev/null
+++ b/Engine/Plugins/Runtime/ApexDestruction/Source/ApexDestructionEditor/Private/DestructibleMeshThumbnailRenderer.cpp
@@ -0,0 +1,54 @@
+// Copyright Epic Games, Inc. All Rights Reserved.
+
+#include "DestructibleMeshThumbnailRenderer.h"
+#include "Misc/App.h"
+#include "ShowFlags.h"
+#include "SceneView.h"
+#include "ThumbnailHelpers.h"
+
+// FPreviewScene derived helpers for rendering
+#include "DestructibleMesh.h"
+#include "DestructibleMeshThumbnailScene.h"
+
+UDestructibleMeshThumbnailRenderer::UDestructibleMeshThumbnailRenderer(const FObjectInitializer& ObjectInitializer)
+	: Super(ObjectInitializer)
+{
+	ThumbnailScene = nullptr;
+}
+
+void UDestructibleMeshThumbnailRenderer::Draw(UObject* Object, int32 X, int32 Y, uint32 Width, uint32 Height, FRenderTarget* RenderTarget, FCanvas* Canvas, bool bAdditionalViewFamily)
+{
+PRAGMA_DISABLE_DEPRECATION_WARNINGS
+	UDestructibleMesh* DestructibleMesh = Cast<UDestructibleMesh>(Object);
+	if (DestructibleMesh != nullptr)
+	{
+		if ( ThumbnailScene == nullptr )
+		{
+			ThumbnailScene = new FDestructibleMeshThumbnailScene();
+		}
+
+		ThumbnailScene->SetDestructibleMesh(DestructibleMesh);
+		FSceneViewFamilyContext ViewFamily( FSceneViewFamily::ConstructionValues( RenderTarget, ThumbnailScene->GetScene(), FEngineShowFlags(ESFIM_Game) )
+			.SetTime(UThumbnailRenderer::GetTime())
+			.SetAdditionalViewFamily(bAdditionalViewFamily));
+
+		ViewFamily.EngineShowFlags.DisableAdvancedFeatures();
+		ViewFamily.EngineShowFlags.MotionBlur = 0;
+		ViewFamily.EngineShowFlags.LOD = 0;
+
+		RenderViewFamily(Canvas, &ViewFamily, ThumbnailScene->CreateView(&ViewFamily, X, Y, Width, Height));
+		ThumbnailScene->SetDestructibleMesh(nullptr);
+	}
+PRAGMA_ENABLE_DEPRECATION_WARNINGS
+}
+
+void UDestructibleMeshThumbnailRenderer::BeginDestroy()
+{
+	if ( ThumbnailScene != nullptr )
+	{
+		delete ThumbnailScene;
+		ThumbnailScene = nullptr;
+	}
+
+	Super::BeginDestroy();
+}
diff --git a/Engine/Plugins/Runtime/ApexDestruction/Source/ApexDestructionEditor/Private/DestructibleMeshThumbnailRenderer.h b/Engine/Plugins/Runtime/ApexDestruction/Source/ApexDestructionEditor/Private/DestructibleMeshThumbnailRenderer.h
new file mode 100644
--- /dev/null
+++ b/Engine/Plugins/Runtime/ApexDestruction/Source/ApexDestructionEditor/Private/DestructibleMeshThumbnailRenderer.h
@@ -0,0 +1,34 @@
+// Copyright Epic Games, Inc. All Rights Reserved.
+
+/**
+ *
+ * This thumbnail renderer displays a given skeletal mesh
+ */
+
+#pragma once
+
+#include "CoreMinimal.h"
+#include "UObject/ObjectMacros.h"
+#include "ThumbnailRendering/DefaultSizedThumbnailRenderer.h"
+#include "DestructibleMeshThumbnailRenderer.generated.h"
+
+class FCanvas;
+class FRenderTarget;
+
+UCLASS(config=Editor, MinimalAPI)
+class UDestructibleMeshThumbnailRenderer : public UDefaultSizedThumbnailRenderer
+{
+	GENERATED_UCLASS_BODY()
+
+
+	// Begin UThumbnailRenderer Object
+	APEXDESTRUCTIONEDITOR_API virtual void Draw(UObject* Object, int32 X, int32 Y, uint32 Width, uint32 Height, FRenderTarget*, FCanvas* Canvas, bool bAdditionalViewFamily) override;
+	// End UThumbnailRenderer Object
+
+	// UObject implementation
+	APEXDESTRUCTIONEDITOR_API virtual void BeginDestroy() override;
+
+private:
+	class FDestructibleMeshThumbnailScene* ThumbnailScene;
+};
+
diff --git a/Engine/Plugins/Runtime/ApexDestruction/Source/ApexDestructionEditor/Private/DestructibleMeshThumbnailScene.cpp b/Engine/Plugins/Runtime/ApexDestruction/Source/ApexDestructionEditor/Private/DestructibleMeshThumbnailScene.cpp
new file mode 100644
--- /dev/null
+++ b/Engine/Plugins/Runtime/ApexDestruction/Source/ApexDestructionEditor/Private/DestructibleMeshThumbnailScene.cpp
@@ -0,0 +1,88 @@
+// Copyright Epic Games, Inc. All Rights Reserved.
+
+#include "DestructibleMeshThumbnailScene.h"
+#include "DestructibleActor.h"
+#include "DestructibleComponent.h"
+#include "DestructibleMesh.h"
+#include "ThumbnailRendering/SceneThumbnailInfo.h"
+
+/*
+***************************************************************
+FDestructibleMeshThumbnailScene
+***************************************************************
+*/
+
+FDestructibleMeshThumbnailScene::FDestructibleMeshThumbnailScene()
+	: FThumbnailPreviewScene()
+{
+	bForceAllUsedMipsResident = false;
+	// Create preview actor
+	// checked
+	FActorSpawnParameters SpawnInfo;
+	SpawnInfo.SpawnCollisionHandlingOverride = ESpawnActorCollisionHandlingMethod::AlwaysSpawn;
+	SpawnInfo.bNoFail = true;
+	SpawnInfo.ObjectFlags = RF_Transient;
+
+PRAGMA_DISABLE_DEPRECATION_WARNINGS
+	PreviewActor = GetWorld()->SpawnActor<ADestructibleActor>(SpawnInfo);
+PRAGMA_ENABLE_DEPRECATION_WARNINGS
+
+	PreviewActor->SetActorEnableCollision(false);
+}
+
+PRAGMA_DISABLE_DEPRECATION_WARNINGS
+void FDestructibleMeshThumbnailScene::SetDestructibleMesh(UDestructibleMesh* InMesh)
+{
+	PreviewActor->GetDestructibleComponent()->OverrideMaterials.Empty();
+	PreviewActor->GetDestructibleComponent()->SetDestructibleMesh(InMesh);
+
+	if (InMesh)
+	{
+		FTransform MeshTransform = FTransform::Identity;
+
+		PreviewActor->SetActorLocation(FVector(0, 0, 0), false);
+		PreviewActor->GetDestructibleComponent()->UpdateBounds();
+
+		// Center the mesh at the world origin then offset to put it on top of the plane
+		const float BoundsZOffset = GetBoundsZOffset(PreviewActor->GetDestructibleComponent()->Bounds);
+		PreviewActor->SetActorLocation(-PreviewActor->GetDestructibleComponent()->Bounds.Origin + FVector(0, 0, BoundsZOffset), false);
+		PreviewActor->GetDestructibleComponent()->RecreateRenderState_Concurrent();
+	}
+}
+PRAGMA_ENABLE_DEPRECATION_WARNINGS
+
+void FDestructibleMeshThumbnailScene::GetViewMatrixParameters(const float InFOVDegrees, FVector& OutOrigin, float& OutOrbitPitch, float& OutOrbitYaw, float& OutOrbitZoom) const
+{
+PRAGMA_DISABLE_DEPRECATION_WARNINGS
+	check(PreviewActor->GetDestructibleComponent());
+	check(PreviewActor->GetDestructibleComponent()->GetDestructibleMesh());
+PRAGMA_ENABLE_DEPRECATION_WARNINGS
+
+	const float HalfFOVRadians = FMath::DegreesToRadians<float>(InFOVDegrees) * 0.5f;
+
+PRAGMA_DISABLE_DEPRECATION_WARNINGS
+	// No need to add extra size to view slightly outside of the sphere to compensate for perspective since skeletal meshes already buffer bounds.
+	const float HalfMeshSize = PreviewActor->GetDestructibleComponent()->Bounds.SphereRadius;
+	const float BoundsZOffset = GetBoundsZOffset(PreviewActor->GetDestructibleComponent()->Bounds);
+	const float TargetDistance = HalfMeshSize / FMath::Tan(HalfFOVRadians);
+
+	USceneThumbnailInfo* ThumbnailInfo = Cast<USceneThumbnailInfo>(PreviewActor->GetDestructibleComponent()->DestructibleMesh_DEPRECATED->GetThumbnailInfo());
+PRAGMA_ENABLE_DEPRECATION_WARNINGS
+
+	if (ThumbnailInfo)
+	{
+		if (TargetDistance + ThumbnailInfo->OrbitZoom < 0)
+		{
+			ThumbnailInfo->OrbitZoom = -TargetDistance;
+		}
+	}
+	else
+	{
+		ThumbnailInfo = USceneThumbnailInfo::StaticClass()->GetDefaultObject<USceneThumbnailInfo>();
+	}
+
+	OutOrigin = FVector(0, 0, -BoundsZOffset);
+	OutOrbitPitch = ThumbnailInfo->OrbitPitch;
+	OutOrbitYaw = ThumbnailInfo->OrbitYaw;
+	OutOrbitZoom = TargetDistance + ThumbnailInfo->OrbitZoom;
+}
diff --git a/Engine/Plugins/Runtime/ApexDestruction/Source/ApexDestructionEditor/Private/DestructibleMeshThumbnailScene.h b/Engine/Plugins/Runtime/ApexDestruction/Source/ApexDestructionEditor/Private/DestructibleMeshThumbnailScene.h
new file mode 100644
--- /dev/null
+++ b/Engine/Plugins/Runtime/ApexDestruction/Source/ApexDestructionEditor/Private/DestructibleMeshThumbnailScene.h
@@ -0,0 +1,33 @@
+// Copyright Epic Games, Inc. All Rights Reserved.
+
+/*
+***************************************************************
+FDestructibleMeshThumbnailScene
+***************************************************************
+*/
+
+#pragma once
+
+#include "ThumbnailHelpers.h"
+
+class APEXDESTRUCTIONEDITOR_API FDestructibleMeshThumbnailScene : public FThumbnailPreviewScene
+{
+public:
+	/** Constructor */
+	FDestructibleMeshThumbnailScene();
+
+PRAGMA_DISABLE_DEPRECATION_WARNINGS
+	/** Sets the skeletal mesh to use in the next CreateView() */
+	void SetDestructibleMesh(class UDestructibleMesh* InMesh);
+PRAGMA_ENABLE_DEPRECATION_WARNINGS
+
+protected:
+	// FThumbnailPreviewScene implementation
+	virtual void GetViewMatrixParameters(const float InFOVDegrees, FVector& OutOrigin, float& OutOrbitPitch, float& OutOrbitYaw, float& OutOrbitZoom) const override;
+
+private:
+PRAGMA_DISABLE_DEPRECATION_WARNINGS
+	/** The skeletal mesh actor used to display all skeletal mesh thumbnails */
+	class ADestructibleActor* PreviewActor;
+PRAGMA_ENABLE_DEPRECATION_WARNINGS
+};
\ No newline at end of file
