Fix errors with latest compiler

#rnx

diff --git a/Engine/Plugins/Runtime/ApexDestruction/Source/ApexDestruction/Public/DestructibleComponent.h b/Engine/Plugins/Runtime/ApexDestruction/Source/ApexDestruction/Public/DestructibleComponent.h
--- a/Engine/Plugins/Runtime/ApexDestruction/Source/ApexDestruction/Public/DestructibleComponent.h
+++ b/Engine/Plugins/Runtime/ApexDestruction/Source/ApexDestruction/Public/DestructibleComponent.h
@@ -55,8 +55,10 @@ class APEXDESTRUCTION_API UDestructibleComponent : public USkinnedMeshComponent,
 
 #if WITH_EDITORONLY_DATA
 	/** Provide a blueprint interface for setting the destructible mesh */
+	PRAGMA_DISABLE_DEPRECATION_WARNINGS
 	UPROPERTY(Transient, BlueprintSetter=SetDestructibleMesh, BlueprintGetter=GetDestructibleMesh, Category = "Components|Destructible", meta = (DeprecatedProperty, DeprecatedMessage = "Destructible mesh is now deprected, use UGeometryCollection instead"))
 	TObjectPtr<class UDestructibleMesh> DestructibleMesh_DEPRECATED;
+	PRAGMA_ENABLE_DEPRECATION_WARNINGS
 #endif // WITH_EDITORONLY_DATA
 
 #if WITH_EDITOR
diff --git a/Engine/Plugins/Runtime/ApexDestruction/Source/ApexDestruction/Public/DestructibleMesh.h b/Engine/Plugins/Runtime/ApexDestruction/Source/ApexDestruction/Public/DestructibleMesh.h
--- a/Engine/Plugins/Runtime/ApexDestruction/Source/ApexDestruction/Public/DestructibleMesh.h
+++ b/Engine/Plugins/Runtime/ApexDestruction/Source/ApexDestruction/Public/DestructibleMesh.h
@@ -393,8 +393,10 @@ struct UE_DEPRECATED(4.26, "APEX is deprecated. Destruction in future will be su
 		the element [0] of the array applies to the level 0 (unfractured) chunk, element [1] applies
 		to the level 1 chunks, etc.
 	*/
+	PRAGMA_DISABLE_DEPRECATION_WARNINGS
 	UPROPERTY(EditAnywhere, editfixedsize, Category=DestructibleParameters)
 	TArray<struct FDestructibleDepthParameters> DepthParameters;
+	PRAGMA_ENABLE_DEPRECATION_WARNINGS
 
 	/** A collection of flags defined in DestructibleParametersFlag. */
 	UPROPERTY(EditAnywhere, Category=DestructibleParameters)
@@ -420,8 +422,10 @@ class UE_DEPRECATED(4.26, "APEX is deprecated. Destruction in future will be sup
 
 #if WITH_EDITORONLY_DATA
 	/** Information used to author an NxDestructibleAsset*/
+	PRAGMA_DISABLE_DEPRECATION_WARNINGS
 	UPROPERTY(instanced)
 	TObjectPtr<class UDestructibleFractureSettings> FractureSettings;
+	PRAGMA_ENABLE_DEPRECATION_WARNINGS
 	
 	/** Static mesh this destructible mesh is created from. Is nullptr if not created from a static mesh */
 	UPROPERTY()
diff --git a/Engine/Source/Runtime/Engine/Classes/Engine/SkeletalMesh.h b/Engine/Source/Runtime/Engine/Classes/Engine/SkeletalMesh.h
--- a/Engine/Source/Runtime/Engine/Classes/Engine/SkeletalMesh.h
+++ b/Engine/Source/Runtime/Engine/Classes/Engine/SkeletalMesh.h
@@ -935,9 +935,11 @@ public:
 
 #if WITH_EDITORONLY_DATA
 	/** List of bones that should be mirrored. */
+	PRAGMA_DISABLE_DEPRECATION_WARNINGS
 	UE_DEPRECATED(4.27, "Please do not access this member directly; Use UMirrorDataTable for mirroring support")
 	UPROPERTY(EditAnywhere, editfixedsize, Category=Mirroring)
 	TArray<struct FBoneMirrorInfo> SkelMirrorTable;
+	PRAGMA_ENABLE_DEPRECATION_WARNINGS
 #endif
 
 	/** Settings related to building Nanite data. */
